Factored out skin weight specific functions so they can be shared by FSkeletalMeshAttributes and FSkeletalMeshConstAttributes.

#trivial

diff --git a/Engine/Source/Runtime/SkeletalMeshDescription/Private/SkeletalMeshAttributes.cpp b/Engine/Source/Runtime/SkeletalMeshDescription/Private/SkeletalMeshAttributes.cpp
--- a/Engine/Source/Runtime/SkeletalMeshDescription/Private/SkeletalMeshAttributes.cpp
+++ b/Engine/Source/Runtime/SkeletalMeshDescription/Private/SkeletalMeshAttributes.cpp
@@ -8,7 +8,7 @@ namespace MeshAttribute
 }
 
 
-FName FSkeletalMeshAttributes::DefaultSkinWeightProfileName("Default");
+FName FSkeletalMeshAttributesShared::DefaultSkinWeightProfileName("Default");
 
 
 static FString SkinWeightAttributeNamePrefix()
@@ -16,7 +16,6 @@ static FString SkinWeightAttributeNamePrefix()
 	return MeshAttribute::Vertex::SkinWeights.ToString() + TEXT("-");
 }
 
-
 void FSkeletalMeshAttributes::Register()
 {
 	MeshDescription.VertexAttributes().RegisterAttribute<int32[]>(MeshAttribute::Vertex::SkinWeights, 1, 0.0f, EMeshAttributeFlags::Mandatory);
@@ -50,10 +49,10 @@ bool FSkeletalMeshAttributes::RegisterSkinWeightAttribute(const FName InProfileN
 	}	
 }
 
-TArray<FName> FSkeletalMeshAttributes::GetSkinWeightProfileNames() const
+TArray<FName> FSkeletalMeshAttributesShared::GetSkinWeightProfileNames() const
 {
 	TArray<FName> AllAttributeNames;
-	MeshDescription.VertexAttributes().GetAttributeNames(AllAttributeNames);
+	MeshDescriptionShared.VertexAttributes().GetAttributeNames(AllAttributeNames);
 
 	TArray<FName> SkinWeightProfileNames;
 	const FString Prefix = SkinWeightAttributeNamePrefix();
@@ -82,15 +81,21 @@ TArray<FName> FSkeletalMeshAttributes::GetSkinWeightProfileNames() const
 	return SkinWeightProfileNames;
 }
 
-bool FSkeletalMeshAttributes::IsValidSkinWeightProfileName(const FName InProfileName)
+bool FSkeletalMeshAttributesShared::IsValidSkinWeightProfileName(const FName InProfileName)
 {
 	return !InProfileName.IsNone() && !InProfileName.IsEqual(DefaultSkinWeightProfileName, ENameCase::IgnoreCase);
 }
 
-FName FSkeletalMeshAttributes::CreateSkinWeightAttributeName(const FName InProfileName)
+bool FSkeletalMeshAttributesShared::IsSkinWeightAttribute(const FName InAttributeName)
+{
+	return InAttributeName == MeshAttribute::Vertex::SkinWeights ||
+		   InAttributeName.ToString().StartsWith(SkinWeightAttributeNamePrefix());
+}
+
+FName FSkeletalMeshAttributesShared::CreateSkinWeightAttributeName(const FName InProfileName)
 {
 	// If it's the default profile, then return the base skin weights attribute name.
-	if (InProfileName.IsNone() || InProfileName.IsEqual(DefaultSkinWeightProfileName, ENameCase::IgnoreCase))
+	if (InProfileName.IsNone() || InProfileName.IsEqual(FSkeletalMeshAttributesShared::DefaultSkinWeightProfileName, ENameCase::IgnoreCase))
 	{
 		return MeshAttribute::Vertex::SkinWeights;
 	}	
diff --git a/Engine/Source/Runtime/SkeletalMeshDescription/Public/SkeletalMeshAttributes.h b/Engine/Source/Runtime/SkeletalMeshDescription/Public/SkeletalMeshAttributes.h
--- a/Engine/Source/Runtime/SkeletalMeshDescription/Public/SkeletalMeshAttributes.h
+++ b/Engine/Source/Runtime/SkeletalMeshDescription/Public/SkeletalMeshAttributes.h
@@ -17,24 +17,15 @@ namespace MeshAttribute
 }
 
 
-
-class SKELETALMESHDESCRIPTION_API FSkeletalMeshAttributes : public FStaticMeshAttributes
+class SKELETALMESHDESCRIPTION_API FSkeletalMeshAttributesShared
 {
 public:
-
-	explicit FSkeletalMeshAttributes(FMeshDescription& InMeshDescription)
-		: FStaticMeshAttributes(InMeshDescription)
-	{}
-
 	// The name of the default skin weight profile.
 	static FName DefaultSkinWeightProfileName;
 	
-	virtual void Register() override;
-
-	/// Register a new skin weight profile with the given name. The attribute name will encode the profile name and
-	/// it will be listed in GetSkinWeightProfileNames(). Returns \c true if the profile was successfully registered.
-	/// Returns \c false if the attribute was already registered or if IsValidSkinWeightProfileName() returned false.
-	bool RegisterSkinWeightAttribute(const FName InProfileName);
+	FSkeletalMeshAttributesShared(const FMeshDescription& InMeshDescription) :
+		MeshDescriptionShared(InMeshDescription)
+	{}
 
 	/// Returns the list of all registered skin weight profile names on this mesh.
 	TArray<FName> GetSkinWeightProfileNames() const;
@@ -43,36 +34,85 @@ public:
 	/// then the profile name is considered invalid. 
 	static bool IsValidSkinWeightProfileName(const FName InProfileName);
 
-	/// Returns the skin weight profile given by its name. NAME_None corresponds to the default profile.
-	FSkinWeightsVertexAttributesRef GetVertexSkinWeights(const FName InProfileName = NAME_None)
+	/// Helper function that indicates whether an attribute name represents a skin weight attribute.
+	static bool IsSkinWeightAttribute(const FName InAttributeName);
+
+	FSkinWeightsVertexAttributesConstRef GetVertexSkinWeights(const FName InProfileName = NAME_None) const
 	{
-		return MeshDescription.VertexAttributes().GetAttributesRef<TArrayAttribute<int32>>(CreateSkinWeightAttributeName(InProfileName));
+		return MeshDescriptionShared.VertexAttributes().GetAttributesRef<TArrayAttribute<int32>>(CreateSkinWeightAttributeName(InProfileName));
 	}
-	FSkinWeightsVertexAttributesConstRef GetVertexSkinWeights(const FName InProfileName = NAME_None) const
+	
+	FSkinWeightsVertexAttributesConstRef GetVertexSkinWeightsFromAttributeName(const FName InAttributeName = NAME_None) const
 	{
-		return MeshDescription.VertexAttributes().GetAttributesRef<TArrayAttribute<int32>>(CreateSkinWeightAttributeName(InProfileName));
+		if (IsSkinWeightAttribute(InAttributeName))
+	{
+			return MeshDescriptionShared.VertexAttributes().GetAttributesRef<TArrayAttribute<int32>>(InAttributeName);
+	}
+		else
+	{
+			return {};
+		}
 	}
 
 protected:
-	friend class FSkeletalMeshConstAttributes;
-
 	/// Construct a name for a skin weight attribute with the given skin weight profile name.
 	/// Each mesh description can hold different skin weight profiles, although the default
 	/// is always present.
 	static FName CreateSkinWeightAttributeName(const FName InProfileName);
+
+private:
+	const FMeshDescription& MeshDescriptionShared;
 };
 
 
-class FSkeletalMeshConstAttributes : public FStaticMeshConstAttributes
+class SKELETALMESHDESCRIPTION_API FSkeletalMeshAttributes :
+	public FStaticMeshAttributes,
+	public FSkeletalMeshAttributesShared
 {
 public:
 
-	explicit FSkeletalMeshConstAttributes(const FMeshDescription& InMeshDescription)
-		: FStaticMeshConstAttributes(InMeshDescription)
+	explicit FSkeletalMeshAttributes(FMeshDescription& InMeshDescription) :
+		FStaticMeshAttributes(InMeshDescription),
+		FSkeletalMeshAttributesShared(InMeshDescription)
 	{}
 
-	FSkinWeightsVertexAttributesConstRef GetVertexSkinWeights(const FName InProfileName = NAME_None) const
+	virtual void Register() override;
+
+	/// Register a new skin weight profile with the given name. The attribute name will encode the profile name and
+	/// it will be listed in GetSkinWeightProfileNames(). Returns \c true if the profile was successfully registered.
+	/// Returns \c false if the attribute was already registered or if IsValidSkinWeightProfileName() returned false.
+	bool RegisterSkinWeightAttribute(const FName InProfileName);
+
+	/// Returns the skin weight profile given by its name. NAME_None corresponds to the default profile.
+	FSkinWeightsVertexAttributesRef GetVertexSkinWeights(const FName InProfileName = NAME_None)
+	{
+		return MeshDescription.VertexAttributes().GetAttributesRef<TArrayAttribute<int32>>(CreateSkinWeightAttributeName(InProfileName));
+	}
+
+	FSkinWeightsVertexAttributesRef GetVertexSkinWeightsFromAttributeName(const FName InAttributeName = NAME_None)
+	{
+		if (IsSkinWeightAttribute(InAttributeName))
+		{
+			return MeshDescription.VertexAttributes().GetAttributesRef<TArrayAttribute<int32>>(InAttributeName);
+		}
+		else
 	{
-		return MeshDescription.VertexAttributes().GetAttributesRef<TArrayAttribute<int32>>(FSkeletalMeshAttributes::CreateSkinWeightAttributeName(InProfileName));
+			return {};
+		}
 	}
+
+};
+
+
+class FSkeletalMeshConstAttributes :
+	public FStaticMeshConstAttributes,
+	public FSkeletalMeshAttributesShared
+{
+public:
+
+	explicit FSkeletalMeshConstAttributes(const FMeshDescription& InMeshDescription) :
+		FStaticMeshConstAttributes(InMeshDescription),
+		FSkeletalMeshAttributesShared(InMeshDescription)
+	{}
+
 };
diff --git a/Engine/Source/Runtime/SkeletalMeshDescription/Public/SkinWeightsAttributesRef.h b/Engine/Source/Runtime/SkeletalMeshDescription/Public/SkinWeightsAttributesRef.h
--- a/Engine/Source/Runtime/SkeletalMeshDescription/Public/SkinWeightsAttributesRef.h
+++ b/Engine/Source/Runtime/SkeletalMeshDescription/Public/SkinWeightsAttributesRef.h
@@ -269,7 +269,7 @@ public:
 
 protected:
 	friend class FSkeletalMeshAttributes;
-	friend class FSkeletalMeshConstAttributes;
+	friend class FSkeletalMeshAttributesShared;
 
 	FSkinWeightsVertexAttributesConstRef(TVertexAttributesRef<TArrayAttribute<int32>> InAttributesConstRef)
         : AttributesConstRef(InAttributesConstRef)
