Added IsReservedAttributeName to check if an attribute is "reserved", i.e. meaning it should only be created by the given *MeshAttributes class or its bases.

#rnx

diff --git a/Engine/Source/Runtime/MeshDescription/Public/MeshAttributes.h b/Engine/Source/Runtime/MeshDescription/Public/MeshAttributes.h
--- a/Engine/Source/Runtime/MeshDescription/Public/MeshAttributes.h
+++ b/Engine/Source/Runtime/MeshDescription/Public/MeshAttributes.h
@@ -55,6 +55,21 @@ public:
 
 	virtual void Register();
 
+	static bool IsReservedAttributeName(const FName InAttributeName)
+	{
+		return InAttributeName == MeshAttribute::Vertex::Position ||
+			   InAttributeName == MeshAttribute::VertexInstance::VertexIndex ||
+			   InAttributeName == MeshAttribute::Edge::VertexIndex ||
+			   InAttributeName == MeshAttribute::Triangle::VertexInstanceIndex ||
+			   InAttributeName == MeshAttribute::Triangle::PolygonIndex ||
+			   InAttributeName == MeshAttribute::Triangle::EdgeIndex ||
+			   InAttributeName == MeshAttribute::Triangle::VertexIndex ||
+			   InAttributeName == MeshAttribute::Triangle::UVIndex ||
+			   InAttributeName == MeshAttribute::Triangle::PolygonGroupIndex ||
+			   InAttributeName == MeshAttribute::UV::UVCoordinate ||
+			   InAttributeName == MeshAttribute::Polygon::PolygonGroupIndex;
+	}
+
 	/** Accessors for cached vertex position array */
 	TVertexAttributesRef<FVector> GetVertexPositions()
 	{
diff --git a/Engine/Source/Runtime/SkeletalMeshDescription/Public/SkeletalMeshAttributes.h b/Engine/Source/Runtime/SkeletalMeshDescription/Public/SkeletalMeshAttributes.h
--- a/Engine/Source/Runtime/SkeletalMeshDescription/Public/SkeletalMeshAttributes.h
+++ b/Engine/Source/Runtime/SkeletalMeshDescription/Public/SkeletalMeshAttributes.h
@@ -45,11 +45,11 @@ public:
 	FSkinWeightsVertexAttributesConstRef GetVertexSkinWeightsFromAttributeName(const FName InAttributeName = NAME_None) const
 	{
 		if (IsSkinWeightAttribute(InAttributeName))
-	{
+		{
 			return MeshDescriptionShared.VertexAttributes().GetAttributesRef<TArrayAttribute<int32>>(InAttributeName);
-	}
+		}
 		else
-	{
+		{
 			return {};
 		}
 	}
@@ -83,6 +83,20 @@ public:
 	/// Returns \c false if the attribute was already registered or if IsValidSkinWeightProfileName() returned false.
 	bool RegisterSkinWeightAttribute(const FName InProfileName);
 
+	static bool IsReservedAttributeName(const FName InAttributeName)
+	{
+		return FStaticMeshAttributes::IsReservedAttributeName(InAttributeName) ||
+               InAttributeName == MeshAttribute::VertexInstance::TextureCoordinate ||
+               InAttributeName == MeshAttribute::VertexInstance::Normal ||
+               InAttributeName == MeshAttribute::VertexInstance::Tangent ||
+               InAttributeName == MeshAttribute::VertexInstance::BinormalSign ||
+               InAttributeName == MeshAttribute::VertexInstance::Color ||
+               InAttributeName == MeshAttribute::Triangle::Normal ||
+               InAttributeName == MeshAttribute::Triangle::Tangent ||
+               InAttributeName == MeshAttribute::Triangle::Binormal ||
+			   IsSkinWeightAttribute(InAttributeName);
+	}	
+	
 	/// Returns the skin weight profile given by its name. NAME_None corresponds to the default profile.
 	FSkinWeightsVertexAttributesRef GetVertexSkinWeights(const FName InProfileName = NAME_None)
 	{
@@ -96,7 +110,7 @@ public:
 			return MeshDescription.VertexAttributes().GetAttributesRef<TArrayAttribute<int32>>(InAttributeName);
 		}
 		else
-	{
+		{
 			return {};
 		}
 	}
diff --git a/Engine/Source/Runtime/StaticMeshDescription/Public/StaticMeshAttributes.h b/Engine/Source/Runtime/StaticMeshDescription/Public/StaticMeshAttributes.h
--- a/Engine/Source/Runtime/StaticMeshDescription/Public/StaticMeshAttributes.h
+++ b/Engine/Source/Runtime/StaticMeshDescription/Public/StaticMeshAttributes.h
@@ -72,6 +72,22 @@ public:
 	{}
 
 	virtual void Register() override;
+	
+	static bool IsReservedAttributeName(const FName InAttributeName)
+	{
+		return FMeshAttributes::IsReservedAttributeName(InAttributeName) ||
+               InAttributeName == MeshAttribute::VertexInstance::TextureCoordinate ||
+               InAttributeName == MeshAttribute::VertexInstance::Normal ||
+               InAttributeName == MeshAttribute::VertexInstance::Tangent ||
+               InAttributeName == MeshAttribute::VertexInstance::BinormalSign ||
+               InAttributeName == MeshAttribute::VertexInstance::Color ||
+               InAttributeName == MeshAttribute::Edge::IsHard ||
+               InAttributeName == MeshAttribute::Triangle::Normal ||
+               InAttributeName == MeshAttribute::Triangle::Tangent ||
+               InAttributeName == MeshAttribute::Triangle::Binormal ||
+               InAttributeName == MeshAttribute::PolygonGroup::ImportedMaterialSlotName
+		;
+	}	
 
 	UE_DEPRECATED(4.26, "Please use RegisterTriangleNormalAndTangentAttributes() instead.")
 	void RegisterPolygonNormalAndTangentAttributes();
