Fix compiler errors in new clang on macOS.

#rnx

diff --git a/Engine/Source/Runtime/SkeletalMeshDescription/Private/SkeletalMeshAttributes.cpp b/Engine/Source/Runtime/SkeletalMeshDescription/Private/SkeletalMeshAttributes.cpp
--- a/Engine/Source/Runtime/SkeletalMeshDescription/Private/SkeletalMeshAttributes.cpp
+++ b/Engine/Source/Runtime/SkeletalMeshDescription/Private/SkeletalMeshAttributes.cpp
@@ -57,7 +57,7 @@
 	TArray<FName> SkinWeightProfileNames;
 	bool bHasDefault = false; 
 
-	for (const FName AttributeName: AllAttributeNames)
+	for (const FName& AttributeName: AllAttributeNames)
 	{
 		if (AttributeName == MeshAttribute::Vertex::SkinWeights)
 		{
diff --git a/Engine/Source/Runtime/SkeletalMeshDescription/Private/SkeletalMeshOperations.cpp b/Engine/Source/Runtime/SkeletalMeshDescription/Private/SkeletalMeshOperations.cpp
--- a/Engine/Source/Runtime/SkeletalMeshDescription/Private/SkeletalMeshOperations.cpp
+++ b/Engine/Source/Runtime/SkeletalMeshDescription/Private/SkeletalMeshOperations.cpp
@@ -24,7 +24,7 @@
 
 	TargetMesh.SuspendVertexIndexing();
 	
-	for (const FVertexID& SourceVertexID : SourceMesh.Vertices().GetElementIDs())
+	for (const FVertexID SourceVertexID : SourceMesh.Vertices().GetElementIDs())
 	{
 		const FVertexID TargetVertexID = FVertexID(AppendSettings.SourceVertexIDOffset + SourceVertexID.GetValue());
 		FVertexBoneWeightsConst SourceBoneWeights = SourceVertexSkinWeights.Get(SourceVertexID);
