Added Bone element and bone attributes (name, parent index, color, pose) to the mesh description via SkeletalMeshAttributes.

diff --git a/Engine/Source/Runtime/MeshDescription/Public/MeshTypes.h b/Engine/Source/Runtime/MeshDescription/Public/MeshTypes.h
--- a/Engine/Source/Runtime/MeshDescription/Public/MeshTypes.h
+++ b/Engine/Source/Runtime/MeshDescription/Public/MeshTypes.h
@@ -290,3 +290,27 @@ struct FPolygonID : public FElementID
 	MESHDESCRIPTION_API static const FPolygonID Invalid;	// @todo mesheditor script: Can we expose these to BP nicely?	Do we even need to?
 };
 
+USTRUCT( BlueprintType )
+struct FBoneID : public FElementID
+{
+	GENERATED_BODY()
+
+	FBoneID()
+	{
+	}
+
+	FBoneID( const FElementID InitElementID )
+		: FElementID( InitElementID.GetValue() )
+	{
+	}
+
+	FBoneID( const int32 InitIDValue )
+		: FElementID( InitIDValue )
+	{
+	}
+
+	FORCEINLINE friend uint32 GetTypeHash( const FBoneID& Other )
+	{
+		return GetTypeHash( Other.IDValue );
+	}
+};
diff --git a/Engine/Source/Runtime/SkeletalMeshDescription/Private/SkeletalMeshAttributes.cpp b/Engine/Source/Runtime/SkeletalMeshDescription/Private/SkeletalMeshAttributes.cpp
--- a/Engine/Source/Runtime/SkeletalMeshDescription/Private/SkeletalMeshAttributes.cpp
+++ b/Engine/Source/Runtime/SkeletalMeshDescription/Private/SkeletalMeshAttributes.cpp
@@ -1,30 +1,73 @@
 // Copyright Epic Games, Inc. All Rights Reserved.
 
 #include "SkeletalMeshAttributes.h"
+#include "MeshDescription.h"
 
 namespace MeshAttribute
 {
 	const FName Vertex::SkinWeights("SkinWeights");
+	
+	namespace Bone
+	{
+		const FName Name("Name");
+		const FName ParentIndex("ParentIndex");
+		const FName Pose("Pose");
+		const FName Color("Color");
+	}
 }
 
 
 FName FSkeletalMeshAttributesShared::DefaultSkinWeightProfileName("Default");
 
+FName FSkeletalMeshAttributesShared::BonesElementName("BonesElementName");
 
 static FString SkinWeightAttributeNamePrefix()
 {
 	return MeshAttribute::Vertex::SkinWeights.ToString() + TEXT("-");
 }
 
+
+//
+// FSkeletalMeshAttributes
+//
+
+FSkeletalMeshAttributes::FSkeletalMeshAttributes(FMeshDescription& InMeshDescription)
+:
+FStaticMeshAttributes(InMeshDescription),
+FSkeletalMeshAttributesShared(InMeshDescription)
+{
+	if (MeshDescription.GetElements().Contains(BonesElementName))
+	{
+		BoneElements = MeshDescription.GetElements()[BonesElementName].Get();
+	}
+}
+
 void FSkeletalMeshAttributes::Register()
 {
 	MeshDescription.VertexAttributes().RegisterAttribute<int32[]>(MeshAttribute::Vertex::SkinWeights, 1, 0, EMeshAttributeFlags::Mandatory);
 
+	if (MeshDescription.GetElements().Contains(BonesElementName) == false)
+	{
+		BoneElements = MeshDescription.GetElements().Emplace(BonesElementName).Get();
+	}
+
+	BoneAttributes().RegisterAttribute<FName>(MeshAttribute::Bone::Name, 1, NAME_None, EMeshAttributeFlags::Mandatory);
+	
+	BoneAttributes().RegisterAttribute<int32>(MeshAttribute::Bone::ParentIndex, 1, INDEX_NONE, EMeshAttributeFlags::Mandatory);
+	
+	BoneAttributes().RegisterAttribute<FTransform>(MeshAttribute::Bone::Pose, 1, FTransform::Identity, EMeshAttributeFlags::Mandatory);
+
 	// Call super class
 	FStaticMeshAttributes::Register();
 }
 
-bool FSkeletalMeshAttributes::RegisterSkinWeightAttribute(const FName InProfileName)
+void FSkeletalMeshAttributes::RegisterColorAttribute()
+{
+	checkSlow(MeshDescription.GetElements().Contains(BonesElementName));
+	BoneAttributes().RegisterAttribute<FVector4f>(MeshAttribute::Bone::Color, 1, FVector4f(1.0f, 1.0f, 1.0f, 1.0f), EMeshAttributeFlags::Mandatory);
+}
+
+bool FSkeletalMeshAttributes::RegisterSkinWeightAttribute(const FName& InProfileName)
 {
 	if (!IsValidSkinWeightProfileName(InProfileName))
 	{
@@ -49,6 +92,106 @@ bool FSkeletalMeshAttributes::RegisterSkinWeightAttribute(const FName InProfileN
 	}	
 }
 
+FSkinWeightsVertexAttributesRef FSkeletalMeshAttributes::GetVertexSkinWeights(const FName& InProfileName)
+{
+	return MeshDescription.VertexAttributes().GetAttributesRef<TArrayAttribute<int32>>(CreateSkinWeightAttributeName(InProfileName));
+}
+
+FSkinWeightsVertexAttributesRef FSkeletalMeshAttributes::GetVertexSkinWeightsFromAttributeName(const FName& InAttributeName)
+{
+	if (IsSkinWeightAttribute(InAttributeName))
+	{
+		return MeshDescription.VertexAttributes().GetAttributesRef<TArrayAttribute<int32>>(InAttributeName);
+	}
+	else
+	{
+		return {};
+	}
+}
+
+FSkeletalMeshAttributes::FBoneArray& FSkeletalMeshAttributes::Bones() 
+{ 
+	return static_cast<TMeshElementContainer<FBoneID>&>(BoneElements->Get()); 
+}
+
+TAttributesSet<FBoneID>& FSkeletalMeshAttributes::BoneAttributes() 
+{ 
+	return Bones().GetAttributes();
+}
+
+void FSkeletalMeshAttributes::ReserveNewBones(const int InNumBones) 
+{
+	BoneElements->Get().Reserve(InNumBones);
+}
+
+FBoneID FSkeletalMeshAttributes::CreateBone()
+{
+	return BoneElements->Get().Add();
+}
+
+void FSkeletalMeshAttributes::CreateBone(const FBoneID BoneID)
+{
+	BoneElements->Get().Insert(BoneID);
+}
+
+void FSkeletalMeshAttributes::DeleteBone(const FBoneID BoneID)
+{
+	BoneElements->Get().Remove(BoneID);
+}
+
+FSkeletalMeshAttributes::FBoneNameAttributesRef FSkeletalMeshAttributes::GetBoneNames()
+{	
+	return BoneAttributes().GetAttributesRef<FName>(MeshAttribute::Bone::Name);
+}
+
+FSkeletalMeshAttributes::FBoneParentIndexAttributesRef FSkeletalMeshAttributes::GetBoneParentIndices()
+{	
+	return BoneAttributes().GetAttributesRef<int32>(MeshAttribute::Bone::ParentIndex);
+}
+
+FSkeletalMeshAttributes::FBonePoseAttributesRef FSkeletalMeshAttributes::GetBonePoses()
+{
+	return BoneAttributes().GetAttributesRef<FTransform>(MeshAttribute::Bone::Pose);
+}
+
+FSkeletalMeshAttributes::FBoneColorAttributesRef FSkeletalMeshAttributes::GetBoneColors()
+{
+	return BoneAttributes().GetAttributesRef<FVector4f>(MeshAttribute::Bone::Color);
+}
+
+
+
+//
+// FSkeletalMeshAttributesShared
+//
+
+FSkeletalMeshAttributesShared::FSkeletalMeshAttributesShared(const FMeshDescription& InMeshDescription) 
+:
+MeshDescriptionShared(InMeshDescription)
+{
+	if (MeshDescriptionShared.GetElements().Contains(BonesElementName))
+	{
+		BoneElementsShared = MeshDescriptionShared.GetElements()[BonesElementName].Get();
+	}
+}
+
+FSkinWeightsVertexAttributesConstRef FSkeletalMeshAttributesShared::GetVertexSkinWeights(const FName& InProfileName) const
+{
+	return MeshDescriptionShared.VertexAttributes().GetAttributesRef<TArrayAttribute<int32>>(CreateSkinWeightAttributeName(InProfileName));
+}
+
+FSkinWeightsVertexAttributesConstRef FSkeletalMeshAttributesShared::GetVertexSkinWeightsFromAttributeName(const FName& InAttributeName) const
+{
+	if (IsSkinWeightAttribute(InAttributeName))
+	{
+		return MeshDescriptionShared.VertexAttributes().GetAttributesRef<TArrayAttribute<int32>>(InAttributeName);
+	}
+	else
+	{
+		return {};
+	}
+}
+
 TArray<FName> FSkeletalMeshAttributesShared::GetSkinWeightProfileNames() const
 {
 	TArray<FName> AllAttributeNames;
@@ -78,18 +221,18 @@ TArray<FName> FSkeletalMeshAttributesShared::GetSkinWeightProfileNames() const
 	return SkinWeightProfileNames;
 }
 
-bool FSkeletalMeshAttributesShared::IsValidSkinWeightProfileName(const FName InProfileName)
+bool FSkeletalMeshAttributesShared::IsValidSkinWeightProfileName(const FName& InProfileName)
 {
 	return !InProfileName.IsNone() && !InProfileName.IsEqual(DefaultSkinWeightProfileName, ENameCase::IgnoreCase);
 }
 
-bool FSkeletalMeshAttributesShared::IsSkinWeightAttribute(const FName InAttributeName)
+bool FSkeletalMeshAttributesShared::IsSkinWeightAttribute(const FName& InAttributeName)
 {
 	return InAttributeName == MeshAttribute::Vertex::SkinWeights ||
 		   InAttributeName.ToString().StartsWith(SkinWeightAttributeNamePrefix());
 }
 
-FName FSkeletalMeshAttributesShared::GetProfileNameFromAttribute(const FName InAttributeName)
+FName FSkeletalMeshAttributesShared::GetProfileNameFromAttribute(const FName& InAttributeName)
 {
 	if (InAttributeName == MeshAttribute::Vertex::SkinWeights)
 	{
@@ -107,7 +250,7 @@ FName FSkeletalMeshAttributesShared::GetProfileNameFromAttribute(const FName InA
 	}
 }
 
-FName FSkeletalMeshAttributesShared::CreateSkinWeightAttributeName(const FName InProfileName)
+FName FSkeletalMeshAttributesShared::CreateSkinWeightAttributeName(const FName& InProfileName)
 {
 	// If it's the default profile, then return the base skin weights attribute name.
 	if (InProfileName.IsNone() || InProfileName.IsEqual(FSkeletalMeshAttributesShared::DefaultSkinWeightProfileName, ENameCase::IgnoreCase))
@@ -117,3 +260,58 @@ FName FSkeletalMeshAttributesShared::CreateSkinWeightAttributeName(const FName I
 
 	return FName(SkinWeightAttributeNamePrefix() + InProfileName.ToString());
 }
+
+bool FSkeletalMeshAttributesShared::HasBoneColorAttribute() const
+{
+	return BoneAttributes().HasAttribute(MeshAttribute::Bone::Color);
+}
+
+bool FSkeletalMeshAttributesShared::HasBoneNameAttribute() const
+{
+	return BoneAttributes().HasAttribute(MeshAttribute::Bone::Name);
+}
+
+bool FSkeletalMeshAttributesShared::HasBonePoseAttribute() const
+{
+	return BoneAttributes().HasAttribute(MeshAttribute::Bone::Pose);
+}
+
+bool FSkeletalMeshAttributesShared::HasBoneParentIndexAttribute() const
+{
+	return BoneAttributes().HasAttribute(MeshAttribute::Bone::ParentIndex);
+}
+
+const FSkeletalMeshAttributesShared::FBoneArray& FSkeletalMeshAttributesShared::Bones() const
+{ 
+	return static_cast<const TMeshElementContainer<FBoneID>&>(BoneElementsShared->Get());
+}
+
+const TAttributesSet<FBoneID>& FSkeletalMeshAttributesShared::BoneAttributes() const 
+{ 
+	return Bones().GetAttributes(); 
+}
+
+int32 FSkeletalMeshAttributesShared::GetNumBones() const
+{
+	return HasBones() ? BoneElementsShared->Get().Num() : 0;
+}
+
+FSkeletalMeshAttributesShared::FBoneNameAttributesConstRef FSkeletalMeshAttributesShared::GetBoneNames() const
+{
+	return BoneAttributes().GetAttributesRef<FName>(MeshAttribute::Bone::Name);
+}
+
+FSkeletalMeshAttributesShared::FBoneParentIndexAttributesConstRef FSkeletalMeshAttributesShared::GetBoneParentIndices() const
+{	
+	return BoneAttributes().GetAttributesRef<int32>(MeshAttribute::Bone::ParentIndex);
+}
+
+FSkeletalMeshAttributesShared::FBonePoseAttributesConstRef FSkeletalMeshAttributesShared::GetBonePoses() const
+{
+	return BoneAttributes().GetAttributesRef<FTransform>(MeshAttribute::Bone::Pose);
+}
+
+FSkeletalMeshAttributesShared::FBoneColorAttributesConstRef FSkeletalMeshAttributesShared::GetBoneColors() const
+{
+	return BoneAttributes().GetAttributesRef<FVector4f>(MeshAttribute::Bone::Color);
+}
\ No newline at end of file
diff --git a/Engine/Source/Runtime/SkeletalMeshDescription/Public/SkeletalMeshAttributes.h b/Engine/Source/Runtime/SkeletalMeshDescription/Public/SkeletalMeshAttributes.h
--- a/Engine/Source/Runtime/SkeletalMeshDescription/Public/SkeletalMeshAttributes.h
+++ b/Engine/Source/Runtime/SkeletalMeshDescription/Public/SkeletalMeshAttributes.h
@@ -6,73 +6,133 @@
 #include "HAL/Platform.h"
 #include "MeshAttributeArray.h"
 #include "MeshAttributes.h"
-#include "MeshDescription.h"
 #include "SkinWeightsAttributesRef.h"
 #include "StaticMeshAttributes.h"
 #include "UObject/NameTypes.h"
 #include "UObject/UnrealNames.h"
+#include "MeshTypes.h"
 
+// Forward declarations
 template <typename AttributeType> class TArrayAttribute;
+struct FMeshDescription;
 
-//Add any skeletalmesh specific attributes here
-
+// Add any skeletalmesh specific attributes here
 namespace MeshAttribute
 {
  	namespace Vertex
  	{
 		extern SKELETALMESHDESCRIPTION_API const FName SkinWeights;
  	}
+
+	namespace Bone
+	{
+		extern SKELETALMESHDESCRIPTION_API const FName Name;
+		extern SKELETALMESHDESCRIPTION_API const FName ParentIndex;
+		extern SKELETALMESHDESCRIPTION_API const FName Pose;
+		extern SKELETALMESHDESCRIPTION_API const FName Color;
+	}
 }
 
 
 class SKELETALMESHDESCRIPTION_API FSkeletalMeshAttributesShared
 {
 public:
+	using FBoneArray = TMeshElementContainer<FBoneID>;
+	
+	using FBoneNameAttributesRef = TMeshAttributesRef<FBoneID, FName>;
+	using FBoneNameAttributesConstRef = TMeshAttributesConstRef<FBoneID, FName>;
+	
+	using FBoneParentIndexAttributesRef = TMeshAttributesRef<FBoneID, int32>;
+	using FBoneParentIndexAttributesConstRef = TMeshAttributesConstRef<FBoneID, int32>;
+	
+	using FBonePoseAttributesRef = TMeshAttributesRef<FBoneID, FTransform>;
+	using FBonePoseAttributesConstRef = TMeshAttributesConstRef<FBoneID, FTransform>;
+	
+	using FBoneColorAttributesRef = TMeshAttributesRef<FBoneID, FVector4f>;
+	using FBoneColorAttributesConstRef = TMeshAttributesConstRef<FBoneID, FVector4f>;
+
+	/** 
+	 * Name of the mesh element type representing bones.
+	 * 
+	 * @note this is different from the MeshAttribute::Bone::Name attribute. This is a name of the element that is 
+	 * added to the  mesh description to represent Bones (similar to the Vertex/Polygons/Edges elements). 
+	 * MeshAttribute::Bone::Name is just one of the attributes of the Bones element.
+	 */
+	static FName BonesElementName;
+
 	// The name of the default skin weight profile.
 	static FName DefaultSkinWeightProfileName;
 	
-	FSkeletalMeshAttributesShared(const FMeshDescription& InMeshDescription) :
-		MeshDescriptionShared(InMeshDescription)
-	{}
+	FSkeletalMeshAttributesShared(const FMeshDescription& InMeshDescription);
+
+	//
+	// Skin Weights Methods
+	//
 
 	/// Returns the list of all registered skin weight profile names on this mesh.
 	TArray<FName> GetSkinWeightProfileNames() const;
 	
 	/// Returns \c true if the given identifier is a valid profile name. If the name is empty, or matches the default profile,
 	/// then the profile name is considered invalid. 
-	static bool IsValidSkinWeightProfileName(const FName InProfileName);
+	static bool IsValidSkinWeightProfileName(const FName& InProfileName);
 
 	/// Helper function that indicates whether an attribute name represents a skin weight attribute.
-	static bool IsSkinWeightAttribute(const FName InAttributeName);
+	static bool IsSkinWeightAttribute(const FName& InAttributeName);
 
 	/// Returns a skin profile name from the attribute name, if the attribute name is a valid skin weights
 	/// attribute.
-	static FName GetProfileNameFromAttribute(const FName InAttributeName);
+	static FName GetProfileNameFromAttribute(const FName& InAttributeName);
 
-	FSkinWeightsVertexAttributesConstRef GetVertexSkinWeights(const FName InProfileName = NAME_None) const
-	{
-		return MeshDescriptionShared.VertexAttributes().GetAttributesRef<TArrayAttribute<int32>>(CreateSkinWeightAttributeName(InProfileName));
-	}
+	FSkinWeightsVertexAttributesConstRef GetVertexSkinWeights(const FName& InProfileName = NAME_None) const;
 	
-	FSkinWeightsVertexAttributesConstRef GetVertexSkinWeightsFromAttributeName(const FName InAttributeName = NAME_None) const
+	FSkinWeightsVertexAttributesConstRef GetVertexSkinWeightsFromAttributeName(const FName& InAttributeName = NAME_None) const;
+
+
+	//
+	// Bones Methods
+	//
+
+	bool HasBoneColorAttribute() const;
+
+	bool HasBoneNameAttribute() const;
+
+	bool HasBonePoseAttribute() const;
+
+	bool HasBoneParentIndexAttribute() const;
+
+	const FBoneArray& Bones() const;
+
+	const TAttributesSet<FBoneID>& BoneAttributes() const;
+
+	/**  @return true, if bone element was added to the MeshDescription. */
+	bool HasBones() const
 	{
-		if (IsSkinWeightAttribute(InAttributeName))
-		{
-			return MeshDescriptionShared.VertexAttributes().GetAttributesRef<TArrayAttribute<int32>>(InAttributeName);
-		}
-		else
-		{
-			return {};
-		}
+		return BoneElementsShared != nullptr;
 	}
 
+	/**  @return the number of the bones. 0 if bone element does not exist in the MeshDescripion. */
+	int32 GetNumBones() const;
+
+	/** @return true, if the passed bone ID is valid */
+	bool IsBoneValid(const FBoneID BoneID) const;
+
+	FBoneNameAttributesConstRef GetBoneNames() const;
+
+	FBoneParentIndexAttributesConstRef GetBoneParentIndices() const;
+
+	FBonePoseAttributesConstRef GetBonePoses() const;
+
+	FBoneColorAttributesConstRef GetBoneColors() const;
+
 protected:
 	/// Construct a name for a skin weight attribute with the given skin weight profile name.
 	/// Each mesh description can hold different skin weight profiles, although the default
 	/// is always present.
-	static FName CreateSkinWeightAttributeName(const FName InProfileName);
+	static FName CreateSkinWeightAttributeName(const FName& InProfileName);
 
 private:
+	const FMeshElementChannels* BoneElementsShared = nullptr;
+
 	const FMeshDescription& MeshDescriptionShared;
 };
 
@@ -83,50 +143,67 @@ class SKELETALMESHDESCRIPTION_API FSkeletalMeshAttributes :
 {
 public:
 
-	explicit FSkeletalMeshAttributes(FMeshDescription& InMeshDescription) :
-		FStaticMeshAttributes(InMeshDescription),
-		FSkeletalMeshAttributesShared(InMeshDescription)
-	{}
+	explicit FSkeletalMeshAttributes(FMeshDescription& InMeshDescription);
 
 	virtual void Register() override;
+	
+	//
+	// Skin Weights Methods
+	//
 
 	/// Register a new skin weight profile with the given name. The attribute name will encode the profile name and
 	/// it will be listed in GetSkinWeightProfileNames(). Returns \c true if the profile was successfully registered.
 	/// Returns \c false if the attribute was already registered or if IsValidSkinWeightProfileName() returned false.
-	bool RegisterSkinWeightAttribute(const FName InProfileName);
+	bool RegisterSkinWeightAttribute(const FName& InProfileName);
 
-	static bool IsReservedAttributeName(const FName InAttributeName)
+	static bool IsReservedAttributeName(const FName& InAttributeName)
 	{
 		return FStaticMeshAttributes::IsReservedAttributeName(InAttributeName) ||
-               InAttributeName == MeshAttribute::VertexInstance::TextureCoordinate ||
-               InAttributeName == MeshAttribute::VertexInstance::Normal ||
-               InAttributeName == MeshAttribute::VertexInstance::Tangent ||
-               InAttributeName == MeshAttribute::VertexInstance::BinormalSign ||
-               InAttributeName == MeshAttribute::VertexInstance::Color ||
-               InAttributeName == MeshAttribute::Triangle::Normal ||
-               InAttributeName == MeshAttribute::Triangle::Tangent ||
-               InAttributeName == MeshAttribute::Triangle::Binormal ||
-			   IsSkinWeightAttribute(InAttributeName);
-	}	
+			   IsSkinWeightAttribute(InAttributeName) ||
+			   InAttributeName == MeshAttribute::Bone::Name ||
+			   InAttributeName == MeshAttribute::Bone::ParentIndex ||
+			   InAttributeName == MeshAttribute::Bone::Pose ||
+			   InAttributeName == MeshAttribute::Bone::Color;
+	}
 	
 	/// Returns the skin weight profile given by its name. NAME_None corresponds to the default profile.
-	FSkinWeightsVertexAttributesRef GetVertexSkinWeights(const FName InProfileName = NAME_None)
-	{
-		return MeshDescription.VertexAttributes().GetAttributesRef<TArrayAttribute<int32>>(CreateSkinWeightAttributeName(InProfileName));
-	}
+	FSkinWeightsVertexAttributesRef GetVertexSkinWeights(const FName& InProfileName = NAME_None);
 
-	FSkinWeightsVertexAttributesRef GetVertexSkinWeightsFromAttributeName(const FName InAttributeName = NAME_None)
-	{
-		if (IsSkinWeightAttribute(InAttributeName))
-		{
-			return MeshDescription.VertexAttributes().GetAttributesRef<TArrayAttribute<int32>>(InAttributeName);
-		}
-		else
-		{
-			return {};
-		}
-	}
+	FSkinWeightsVertexAttributesRef GetVertexSkinWeightsFromAttributeName(const FName& InAttributeName = NAME_None);
+
+	//
+	// Bones Methods
+	//
 
+	/** Register an optional color attribute for bones */
+	void RegisterColorAttribute();
+
+	FBoneArray& Bones();
+
+	TAttributesSet<FBoneID>& BoneAttributes();
+
+	/** Reserves space for this number of new bones */
+	void ReserveNewBones(const int InNumBones);
+
+	/** Adds a new bone and returns its ID */
+	FBoneID CreateBone();
+
+	/** Adds a new bone  with the given ID */
+	void CreateBone(const FBoneID BoneID);
+
+	/** Deletes a bone with the given ID */
+	void DeleteBone(const FBoneID BoneID);
+
+	FBoneNameAttributesRef GetBoneNames();
+	
+	FBoneParentIndexAttributesRef GetBoneParentIndices();
+
+	FBonePoseAttributesRef GetBonePoses();
+
+	FBoneColorAttributesRef GetBoneColors();
+
+private:
+	FMeshElementChannels* BoneElements = nullptr;
 };
 
 
