Fix CIS from bad merge
#rnx

diff --git a/Engine/Source/Runtime/SkeletalMeshDescription/Private/SkeletalMeshAttributes.cpp b/Engine/Source/Runtime/SkeletalMeshDescription/Private/SkeletalMeshAttributes.cpp
--- a/Engine/Source/Runtime/SkeletalMeshDescription/Private/SkeletalMeshAttributes.cpp
+++ b/Engine/Source/Runtime/SkeletalMeshDescription/Private/SkeletalMeshAttributes.cpp
@@ -54,20 +54,11 @@ void FSkeletalMeshAttributes::Register(bool bKeepExistingAttribute)
 		BoneElements = MeshDescription.GetElements().Emplace(BonesElementName).Get();
 	}
 
-	if (!MeshDescription.BoneAttributes().HasAttribute(MeshAttribute::Bone::Name) || !bKeepExistingAttribute)
-	{
-		BoneAttributes().RegisterAttribute<FName>(MeshAttribute::Bone::Name, 1, NAME_None, EMeshAttributeFlags::Mandatory);
-	}
+	BoneAttributes().RegisterAttribute<FName>(MeshAttribute::Bone::Name, 1, NAME_None, EMeshAttributeFlags::Mandatory);
 	
-	if (!MeshDescription.BoneAttributes().HasAttribute(MeshAttribute::Bone::ParentIndex) || !bKeepExistingAttribute)
-	{
-		BoneAttributes().RegisterAttribute<int32>(MeshAttribute::Bone::ParentIndex, 1, INDEX_NONE, EMeshAttributeFlags::Mandatory);
-	}
+	BoneAttributes().RegisterAttribute<int32>(MeshAttribute::Bone::ParentIndex, 1, INDEX_NONE, EMeshAttributeFlags::Mandatory);
 	
-	if (!MeshDescription.BoneAttributes().HasAttribute(MeshAttribute::Bone::Pose) || !bKeepExistingAttribute)
-	{
-		BoneAttributes().RegisterAttribute<FTransform>(MeshAttribute::Bone::Pose, 1, FTransform::Identity, EMeshAttributeFlags::Mandatory);
-	}
+	BoneAttributes().RegisterAttribute<FTransform>(MeshAttribute::Bone::Pose, 1, FTransform::Identity, EMeshAttributeFlags::Mandatory);
 
 	// Call super class
 	FStaticMeshAttributes::Register(bKeepExistingAttribute);
