Skeletal Mesh: Change the import storage to use FMeshDescription instead of FSkeletalMeshImportData. This includes deprecating the USkeletalMesh::LoadLODImportedData and friends, in lieu of an API mor

diff --git a/Engine/Source/Runtime/SkeletalMeshDescription/Public/MorphTargetAttributesRef.h b/Engine/Source/Runtime/SkeletalMeshDescription/Public/MorphTargetAttributesRef.h
--- a/Engine/Source/Runtime/SkeletalMeshDescription/Public/MorphTargetAttributesRef.h
+++ b/Engine/Source/Runtime/SkeletalMeshDescription/Public/MorphTargetAttributesRef.h
@@ -6,6 +6,9 @@
 
 class FMorphTargetVertexAttributesRef
 {
+	friend class FSkeletalMeshAttributes;
+	friend class FMorphTargetVertexAttributesConstRef;
+
 public:
 	FMorphTargetVertexAttributesRef() = default;
 	FMorphTargetVertexAttributesRef(const FMorphTargetVertexAttributesRef&) = default;
@@ -18,6 +21,8 @@ public:
 		return AttributesRef.IsValid();
 	}
 
+	void Copy(const FMorphTargetVertexAttributesConstRef& InSourceAttribute);
+
 	FVector3f GetPositionDelta(const FVertexID InVertexIndex) const
 	{
 		return AttributesRef.Get(InVertexIndex)[0]; 
@@ -54,14 +59,11 @@ public:
 	}
 
 protected:
-	friend class FSkeletalMeshAttributes;
-
 	FMorphTargetVertexAttributesRef(TVertexAttributesRef<TArrayView<FVector3f>> InAttributesRef)
 		: AttributesRef(InAttributesRef)
 	{}
 
 private:
-	friend class FMorphTargetVertexAttributesConstRef;
 	TVertexAttributesRef<TArrayView<FVector3f>> AttributesRef;
 };
 
@@ -76,11 +78,11 @@ public:
 	FMorphTargetVertexAttributesConstRef& operator=(FMorphTargetVertexAttributesConstRef&&) = default;
 
 	// Converting constructors from the non-const variant
-	explicit FMorphTargetVertexAttributesConstRef(const FMorphTargetVertexAttributesRef &InAttributesRef) :
+	explicit FMorphTargetVertexAttributesConstRef(const FMorphTargetVertexAttributesRef& InAttributesRef) :
 		AttributesConstRef(InAttributesRef.AttributesRef)
 	{}
 	
-	FMorphTargetVertexAttributesConstRef& operator=(const FMorphTargetVertexAttributesRef &InAttributesRef)
+	FMorphTargetVertexAttributesConstRef& operator=(const FMorphTargetVertexAttributesRef& InAttributesRef)
 	{
 		AttributesConstRef = InAttributesRef.AttributesRef;
 		return *this;
@@ -102,6 +104,7 @@ public:
 	}
 
 protected:
+	friend class FMorphTargetVertexAttributesRef;
 	friend class FSkeletalMeshAttributes;
 	friend class FSkeletalMeshAttributesShared;
 
@@ -116,3 +119,9 @@ protected:
 private:
 	TVertexAttributesConstRef<TArrayView<FVector3f>> AttributesConstRef;
 };
+
+
+inline void FMorphTargetVertexAttributesRef::Copy(const FMorphTargetVertexAttributesConstRef& InSourceAttribute)
+{
+	AttributesRef.Copy(InSourceAttribute.AttributesConstRef);
+}
diff --git a/Engine/Source/Runtime/SkeletalMeshDescription/Public/SkeletalMeshAttributes.h b/Engine/Source/Runtime/SkeletalMeshDescription/Public/SkeletalMeshAttributes.h
--- a/Engine/Source/Runtime/SkeletalMeshDescription/Public/SkeletalMeshAttributes.h
+++ b/Engine/Source/Runtime/SkeletalMeshDescription/Public/SkeletalMeshAttributes.h
@@ -2,6 +2,8 @@
 
 #pragma once
 
+#include "SkeletalMeshElementTypes.h"
+
 #include "Containers/Array.h"
 #include "HAL/Platform.h"
 #include "MeshAttributeArray.h"
@@ -10,7 +12,6 @@
 #include "StaticMeshAttributes.h"
 #include "UObject/NameTypes.h"
 #include "UObject/UnrealNames.h"
-#include "MeshTypes.h"
 
 // Forward declarations
 template <typename AttributeType> class TArrayAttribute;
@@ -23,6 +24,8 @@ namespace MeshAttribute
  	{
  		// Name of the default skin weights attribute.
 		extern SKELETALMESHDESCRIPTION_API const FName SkinWeights;
+
+ 		extern SKELETALMESHDESCRIPTION_API const FName ImportPointIndex;
  	}
 
 	namespace Bone
@@ -32,6 +35,12 @@ namespace MeshAttribute
 		extern SKELETALMESHDESCRIPTION_API const FName Pose;
 		extern SKELETALMESHDESCRIPTION_API const FName Color;
 	}
+
+	namespace SourceGeometryPart
+ 	{
+ 		extern SKELETALMESHDESCRIPTION_API const FName Name;
+ 		extern SKELETALMESHDESCRIPTION_API const FName VertexOffsetAndCount;
+ 	}
 }
 
 
@@ -52,6 +61,14 @@ public:
 	using FBoneColorAttributesRef = TMeshAttributesRef<FBoneID, FVector4f>;
 	using FBoneColorAttributesConstRef = TMeshAttributesConstRef<FBoneID, FVector4f>;
 
+	using FSourceGeometryPartArray = TMeshElementContainer<FSourceGeometryPartID>;
+	
+	using FSourceGeometryPartNameRef = TMeshAttributesRef<FSourceGeometryPartID, FName>;
+	using FSourceGeometryPartNameConstRef = TMeshAttributesConstRef<FSourceGeometryPartID, FName>;
+	
+	using FSourceGeometryPartVertexOffsetAndCountRef = TMeshAttributesRef<FSourceGeometryPartID, TArrayView<int32>>;
+	using FSourceGeometryPartVertexOffsetAndCountConstRef = TMeshAttributesConstRef<FSourceGeometryPartID, TArrayView<int32>>;
+	
 	/** 
 	 * Name of the mesh element type representing bones.
 	 * 
@@ -61,6 +78,15 @@ public:
 	 */
 	static SKELETALMESHDESCRIPTION_API FName BonesElementName;
 
+	
+	/**
+	 * Name of the mesh element representing import data for meshes that were imported from separate geometries
+	 * into a single skeletal mesh piece. Can be used to match parts of the mesh back to the constituent parts
+	 * of the import source.
+	 */
+	static SKELETALMESHDESCRIPTION_API FName SourceGeometryPartElementName;
+	
+	
 	// The name of the default skin weight profile.
 	static SKELETALMESHDESCRIPTION_API FName DefaultSkinWeightProfileName;
 	
@@ -141,6 +167,25 @@ public:
 
 	SKELETALMESHDESCRIPTION_API FBoneColorAttributesConstRef GetBoneColors() const;
 
+
+	//
+	bool HasSourceGeometryParts() const
+	{
+		return SourceGeometryPartElementsShared != nullptr;
+	}
+
+	SKELETALMESHDESCRIPTION_API const FSourceGeometryPartArray& SourceGeometryParts() const;
+	
+	SKELETALMESHDESCRIPTION_API const TAttributesSet<FSourceGeometryPartID>& SourceGeometryPartAttributes() const;
+		
+	SKELETALMESHDESCRIPTION_API int32 GetNumSourceGeometryParts() const;
+
+	SKELETALMESHDESCRIPTION_API bool IsSourceGeometryPartValid(const FSourceGeometryPartID InSourceGeometryPartID) const;
+
+	SKELETALMESHDESCRIPTION_API FSourceGeometryPartNameConstRef GetSourceGeometryPartNames() const;
+	
+	SKELETALMESHDESCRIPTION_API FSourceGeometryPartVertexOffsetAndCountConstRef GetSourceGeometryPartVertexOffsetAndCounts() const;
+	
 protected:
 	/// Construct a name for a skin weight attribute with the given skin weight profile name.
 	/// Each mesh description can hold different skin weight profiles, although the default
@@ -150,9 +195,10 @@ protected:
 	/// Construct a name for a morph target attribute with the given a user-visible morph target name.
 	static SKELETALMESHDESCRIPTION_API FName CreateMorphTargetAttributeName(const FName InMorphTargetName);
 	
-private:
 	const FMeshElementChannels* BoneElementsShared = nullptr;
+	const FMeshElementChannels* SourceGeometryPartElementsShared = nullptr;
 
+private:
 	const FMeshDescription& MeshDescriptionShared;
 };
 
@@ -167,6 +213,9 @@ public:
 
 	SKELETALMESHDESCRIPTION_API virtual void Register(bool bKeepExistingAttribute = false) override;
 
+	/** Returns \c true if a given attribute name is a name for a reserved attribute or not. If not a reserved
+	 *  attribute, it's one that has been user-defined and is not required by any system.
+	 */
 	static bool IsReservedAttributeName(const FName InAttributeName)
 	{
 		return FStaticMeshAttributes::IsReservedAttributeName(InAttributeName) ||
@@ -178,6 +227,19 @@ public:
 			   InAttributeName == MeshAttribute::Bone::Color;
 	}
 	
+	//
+	/** Register an attribute that can be used to maintain a relationship between a given vertex and the point it originated from
+	 *  on the original imported mesh. This can be useful when trying to match the stored mesh to the import mesh as stored in the import file
+	 *  (e.g. FBX, Alembic, glTF, etc).
+	 */
+	SKELETALMESHDESCRIPTION_API bool RegisterImportPointIndexAttribute();
+
+	/** Convenience function to unregister the attribute that maintains the relationship between a given vertex and the originating point
+	 *  on the import mesh. Use if updating the mesh description but no longer maintaining this relationship (e.g. the topology changed and
+	 *  so this map cannot possibly be accurate).
+	 */
+	SKELETALMESHDESCRIPTION_API void UnregisterImportPointIndexAttribute();
+	
 	//
 	// Skin Weights Methods
 	//
@@ -200,9 +262,19 @@ public:
 	/// it will be listed in GetMorphTargetNames(). Returns \c true if the morph target was successfully registered.
 	/// Returns \c false if the attribute was already registered or if \c InMorphTargetName is empty.
 	SKELETALMESHDESCRIPTION_API bool RegisterMorphTargetAttribute(const FName InMorphTargetName);
+
+	/// Unregister an existing morph target with the given name (as returned by GetMorphTargetNames()). Returns \c true if the morph target
+	/// was successfully unregistered.
+	/// Returns \c false if the attribute wasn't registered or if \c InMorphTargetName is empty.
+	SKELETALMESHDESCRIPTION_API bool UnregisterMorphTargetAttribute(const FName InMorphTargetName);
 	
+	/// Returns a specialized morph target vertex attribute, given a morph target name, to allow setting morph target data
+	/// in a simple fashion. If the morph target doesn't exist, the returned  attribute will be invalid and cannot be read from or written to. 
 	SKELETALMESHDESCRIPTION_API FMorphTargetVertexAttributesRef GetVertexMorphTarget(const FName InMorphTargetName = NAME_None);
 	
+	/// Returns a specialized morph target vertex attribute, given an attribute name, to allow setting morph target data
+	/// in a simple fashion. If the attribute doesn't exist, or the attribute does not refer to a morph target, the returned attribute
+	/// will be invalid and cannot be read from or written to. 
 	SKELETALMESHDESCRIPTION_API FMorphTargetVertexAttributesRef GetVertexMorphTargetFromAttributeName(const FName InAttributeName = NAME_None);
 	
 	//
@@ -236,8 +308,31 @@ public:
 
 	SKELETALMESHDESCRIPTION_API FBoneColorAttributesRef GetBoneColors();
 
+
+	/* Source Geometry Parts */
+
+	/** This will register the source geometry attributes and their element container on the mesh description so that
+	 * source geometry part information can be stored.
+	 */
+	SKELETALMESHDESCRIPTION_API void RegisterSourceGeometryPartsAttributes();
+
+	SKELETALMESHDESCRIPTION_API FSourceGeometryPartArray& SourceGeometryParts();
+	
+	SKELETALMESHDESCRIPTION_API TAttributesSet<FSourceGeometryPartID>& SourceGeometryPartAttributes();
+	
+	SKELETALMESHDESCRIPTION_API FSourceGeometryPartID CreateSourceGeometryPart();
+
+	SKELETALMESHDESCRIPTION_API void DeleteSourceGeometryPart(FSourceGeometryPartID InSourceGeometryPartID);
+
+	SKELETALMESHDESCRIPTION_API FSourceGeometryPartNameRef GetSourceGeometryPartNames();
+	
+	SKELETALMESHDESCRIPTION_API FSourceGeometryPartVertexOffsetAndCountRef GetSourceGeometryPartVertexOffsetAndCounts();
+	
+	
+	
 private:
 	FMeshElementChannels* BoneElements = nullptr;
+	FMeshElementChannels* SourceGeometryPartElements = nullptr;
 };
 
 
