Mesh Description: Utility to get a skeletal source mesh in a given pose.

- Pose can be set by either using component space transforms, or overriding specific bone space transforms.
- Morph target weights can be set, and morph target include morph normal deltas, if present.
- Resulting mesh stores the pose given so that the mesh can be re-posed again.

diff --git a/Engine/Source/Runtime/SkeletalMeshDescription/Private/SkeletalMeshOperations.cpp b/Engine/Source/Runtime/SkeletalMeshDescription/Private/SkeletalMeshOperations.cpp
--- a/Engine/Source/Runtime/SkeletalMeshDescription/Private/SkeletalMeshOperations.cpp
+++ b/Engine/Source/Runtime/SkeletalMeshDescription/Private/SkeletalMeshOperations.cpp
@@ -268,5 +268,358 @@ bool FSkeletalMeshOperations::RemapBoneIndicesOnSkinWeightAttribute(FMeshDescrip
 	return true;
 }
 
+namespace Impl
+{
+static void GetPosedMesh(
+	const FMeshDescription& InSourceMesh, 
+	FMeshDescription& OutTargetMesh,
+	TConstArrayView<FMatrix44f> InRefToUserTransforms, 
+	const FName InSkinWeightProfile, 
+	const TMap<FName, float>& InMorphTargetWeights
+	)
+{
+	struct FMorphInfo
+	{
+		TVertexAttributesRef<FVector3f> PositionDelta;
+		TVertexInstanceAttributesRef<FVector3f> NormalDelta;
+		const float Weight = 0.0f;
+	};
+	
+	OutTargetMesh = InSourceMesh;
+	FSkeletalMeshAttributes Attributes(OutTargetMesh);
+
+	// We need the mesh to be compact for the parallel for to work.
+	FElementIDRemappings Remappings;
+	OutTargetMesh.Compact(Remappings);
+
+	TVertexAttributesRef<FVector3f> PositionAttribute = Attributes.GetVertexPositions();
+	TVertexInstanceAttributesRef<FVector3f> NormalAttribute = Attributes.GetVertexInstanceNormals();
+	TVertexInstanceAttributesRef<FVector3f> TangentAttribute = Attributes.GetVertexInstanceTangents();
+	TVertexInstanceAttributesRef<float> BinormalSignsAttribute = Attributes.GetVertexInstanceBinormalSigns();
+
+	// See which morph target attributes we can peel out. If the normal attributes are not all valid, then
+	// we have to automatically compute the normal from the positions. Otherwise, we only use the normal deltas. 
+	TArray<FMorphInfo> MorphInfos;
+	bool bAllMorphNormalsValid = true;
+	for (const TPair<FName, float>& Item: InMorphTargetWeights)
+	{
+		const FName MorphName{Item.Key};
+		const float MorphWeight{Item.Value};
+		TVertexAttributesRef<FVector3f> PositionDelta = Attributes.GetVertexMorphPositionDelta(MorphName);
+		// Q: Should we use the value of `r.MorphTarget.WeightThreshold` instead? The following condition is
+		// identical to the default setting of that value.
+		if (PositionDelta.IsValid() && !FMath::IsNearlyZero(MorphWeight))
+		{
+			TVertexInstanceAttributesRef<FVector3f> NormalDelta = Attributes.GetVertexInstanceMorphNormalDelta(MorphName);
+			if (!NormalDelta.IsValid())
+			{
+				bAllMorphNormalsValid = false;
+			}
+			MorphInfos.Emplace(PositionDelta, NormalDelta, MorphWeight);
+		}
+	}
+
+	// First we apply the morph info on the positions and normals.
+	if (!MorphInfos.IsEmpty())
+	{
+		struct FMorphProcessContext
+		{
+			TSet<FVertexInstanceID> DirtyVertexInstances;
+			TArray<FVertexID> Neighbors;
+		};
+		TArray<FMorphProcessContext> Contexts;
+		ParallelForWithTaskContext(Contexts, OutTargetMesh.Vertices().Num(),
+			[&Mesh=OutTargetMesh, &PositionAttribute, &MorphInfos, bAllMorphNormalsValid](FMorphProcessContext& Context, int32 Index)
+			{
+				const FVertexID VertexID{Index};
+
+				FVector3f Position = PositionAttribute.Get(VertexID);
+				bool bMoved = false;
+				for (const FMorphInfo& MorphInfo: MorphInfos)
+				{
+					const FVector3f PositionDelta = MorphInfo.PositionDelta.Get(VertexID) * MorphInfo.Weight;
+					if (!PositionDelta.IsNearlyZero())
+					{
+						Position += PositionDelta;
+						bMoved = true;
+					}
+				}
+
+				// If we need to re-generate the normals, store which vertices got moved _and_ their neighbors, since the whole
+				// triangle moved, which affects neighboring vertices of the moved vertex.
+				if (bMoved)
+				{
+					PositionAttribute.Set(VertexID, Position);
+					
+					if (!bAllMorphNormalsValid)
+					{
+						Context.DirtyVertexInstances.Append(Mesh.GetVertexVertexInstanceIDs(VertexID));
+
+						Mesh.GetVertexAdjacentVertices(VertexID, Context.Neighbors);
+						for (const FVertexID NeighborVertexID: Context.Neighbors)
+						{
+							Context.DirtyVertexInstances.Append(Mesh.GetVertexVertexInstanceIDs(NeighborVertexID));
+						}
+					}
+				}
+			});
+
+		if (bAllMorphNormalsValid)
+		{
+			ParallelForWithTaskContext(Contexts, OutTargetMesh.VertexInstances().Num(),
+				[&MorphInfos, &NormalAttribute, &TangentAttribute, &BinormalSignsAttribute](FMorphProcessContext& Context, int32 Index)
+				{
+					FVertexInstanceID VertexInstanceID{Index};
+
+					FVector3f Normal = NormalAttribute.Get(VertexInstanceID);
+					FVector3f Tangent = TangentAttribute.Get(VertexInstanceID);
+					FVector3f Binormal = FVector3f::CrossProduct(Normal, Tangent) * BinormalSignsAttribute.Get(VertexInstanceID);
+					
+					bool bMoved = false;
+					for (const FMorphInfo& MorphInfo: MorphInfos)
+					{
+						const FVector3f NormalDelta = MorphInfo.NormalDelta.Get(VertexInstanceID) * MorphInfo.Weight;
+						if (!NormalDelta.IsNearlyZero())
+						{
+							Normal += NormalDelta;
+							bMoved = true;
+						}
+					}
+
+					if (bMoved)
+					{
+						if (Normal.Normalize())
+						{
+							// Badly named function. This orthonormalizes X & Y using Z as the control. 
+							FVector3f::CreateOrthonormalBasis(Tangent, Binormal, Normal);
+							
+							NormalAttribute.Set(VertexInstanceID, Normal);
+							TangentAttribute.Set(VertexInstanceID, Tangent);
+							const float BinormalSign = FMatrix44f(Tangent, Binormal, Normal, FVector3f::ZeroVector).Determinant() < 0 ? -1.0f : +1.0f;
+							BinormalSignsAttribute.Set(VertexInstanceID, BinormalSign);
+						}
+						else
+						{
+							// Something went wrong. Reconstruct the normal from the tangent and binormal.
+							Context.DirtyVertexInstances.Add(VertexInstanceID);
+						}
+					}
+				});
+		}
+		
+		// Clear out any normals that were affected by the point move, or ended up being degenerate during normal offsetting.
+		TSet<FVertexInstanceID> DirtyVertexInstances;
+		for (const FMorphProcessContext& ProcessContext: Contexts)
+		{
+			DirtyVertexInstances.Append(ProcessContext.DirtyVertexInstances);
+		}
+		
+		if (!DirtyVertexInstances.IsEmpty())
+		{
+			// Mark any vector as zero that we want to regenerate from triangle + neighbors + tangents.
+			for (const FVertexInstanceID VertexInstanceID: DirtyVertexInstances)
+			{
+				NormalAttribute.Set(VertexInstanceID, FVector3f::ZeroVector);
+			}
+
+			// Compute the normals on the dirty vertices, and adjust the tangents to match.
+			FSkeletalMeshOperations::ComputeTangentsAndNormals(OutTargetMesh, EComputeNTBsFlags::WeightedNTBs);
+		}
+	}
+
+	using namespace UE::AnimationCore;
+
+	// The normal needs to be transformed using the inverse transpose of the transform matrices to ensure that
+	// scaling works correctly. 
+	TArray<FMatrix44f> RefToUserTransformsNormal;
+	RefToUserTransformsNormal.Reserve(InRefToUserTransforms.Num());
+	for (const FMatrix44f& Mat: InRefToUserTransforms)
+	{
+		RefToUserTransformsNormal.Add(Mat.Inverse().GetTransposed());
+	}
+	
+	FSkinWeightsVertexAttributesRef SkinWeightAttribute = Attributes.GetVertexSkinWeights(InSkinWeightProfile);
+	ParallelFor(OutTargetMesh.Vertices().Num(),
+		[&Mesh=OutTargetMesh, &PositionAttribute, &NormalAttribute, &TangentAttribute, &SkinWeightAttribute, &RefToUserTransforms=InRefToUserTransforms, &RefToUserTransformsNormal](int32 Index)
+		{
+			const FVertexID VertexID(Index);
+			const FVertexBoneWeights BoneWeights = SkinWeightAttribute.Get(VertexID);
+			const FVector3f Position = PositionAttribute.Get(VertexID);
+			FVector3f SkinnedPosition = FVector3f::ZeroVector;
+
+			for (const FBoneWeight BW: BoneWeights)
+			{
+				SkinnedPosition += RefToUserTransforms[BW.GetBoneIndex()].TransformPosition(Position) * BW.GetWeight();
+			}
+			PositionAttribute.Set(VertexID, SkinnedPosition);
+			
+			for (const FVertexInstanceID VertexInstanceID: Mesh.GetVertexVertexInstanceIDs(VertexID))
+			{
+				const FVector3f Normal = NormalAttribute.Get(VertexInstanceID);
+				const FVector3f Tangent = TangentAttribute.Get(VertexInstanceID);
+				FVector3f SkinnedNormal = FVector3f::ZeroVector;
+				FVector3f SkinnedTangent = FVector3f::ZeroVector;
+
+				for (const FBoneWeight BW: BoneWeights)
+				{
+					SkinnedNormal += RefToUserTransformsNormal[BW.GetBoneIndex()].TransformVector(Normal) * BW.GetWeight();
+					SkinnedTangent += RefToUserTransforms[BW.GetBoneIndex()].TransformVector(Tangent) * BW.GetWeight();
+				}
+				
+				SkinnedNormal.Normalize();
+				SkinnedTangent.Normalize();
+
+				NormalAttribute.Set(VertexInstanceID, SkinnedNormal);
+				TangentAttribute.Set(VertexInstanceID, SkinnedTangent);
+			}
+		});
+}
+
+} // namespace Impl
+
+bool FSkeletalMeshOperations::GetPosedMesh(
+	const FMeshDescription& InSourceMesh, 
+	FMeshDescription& OutTargetMesh,
+	TConstArrayView<FTransform> InComponentSpaceTransforms, 
+	const FName InSkinWeightProfile, 
+	const TMap<FName, float>& InMorphTargetWeights
+	)
+{
+	// Verify that the mesh is valid.
+	FSkeletalMeshConstAttributes Attributes(InSourceMesh);
+	if (!Attributes.HasBonePoseAttribute() || !Attributes.HasBoneParentIndexAttribute())
+	{
+		return false;
+	}
+
+	if (!Attributes.GetVertexSkinWeights(InSkinWeightProfile).IsValid())
+	{
+		return false;
+	}
+
+	if (Attributes.GetNumBones() != InComponentSpaceTransforms.Num())
+	{
+		return false;
+	}
+
+	// Convert the component-space transforms into a set of matrices that transform from the reference pose to
+	// the user pose. These are then used to nudge the vertices from the reference pose to the wanted
+	// user pose by weighing the influence of each bone on a given vertex. If the user pose and the reference pose
+	// are identical, these are all identity matrices.
+	TArray<FMatrix44f> RefToUserTransforms;
+	TArray<FMatrix44f> RefPoseTransforms;
+	
+	FSkeletalMeshAttributes::FBonePoseAttributesConstRef BonePoseAttribute = Attributes.GetBonePoses(); 
+	FSkeletalMeshAttributes::FBoneParentIndexAttributesConstRef ParentBoneIndexAttribute = Attributes.GetBoneParentIndices(); 
+	RefToUserTransforms.SetNumUninitialized(Attributes.GetNumBones());
+	RefPoseTransforms.SetNumUninitialized(Attributes.GetNumBones());
+	
+	for (int32 BoneIndex = 0; BoneIndex < Attributes.GetNumBones(); BoneIndex++)
+	{
+		const int32 ParentBoneIndex = ParentBoneIndexAttribute.Get(BoneIndex);
+		RefPoseTransforms[BoneIndex] = FMatrix44f{BonePoseAttribute.Get(BoneIndex).ToMatrixWithScale()};
+
+		if (ParentBoneIndex != INDEX_NONE)
+		{
+			RefPoseTransforms[BoneIndex] = RefPoseTransforms[BoneIndex] * RefPoseTransforms[ParentBoneIndex];
+		}
+
+		RefToUserTransforms[BoneIndex] = RefPoseTransforms[BoneIndex].Inverse() * FMatrix44f{InComponentSpaceTransforms[BoneIndex].ToMatrixWithScale()};
+	}
+	
+	Impl::GetPosedMesh(InSourceMesh, OutTargetMesh, RefToUserTransforms, InSkinWeightProfile, InMorphTargetWeights);
+
+	// Write out the current ref pose (in bone-space) to the mesh. 
+	FSkeletalMeshAttributes WriteAttributes(OutTargetMesh);
+	FSkeletalMeshAttributes::FBonePoseAttributesRef WriteBonePoseAttribute = WriteAttributes.GetBonePoses(); 
+	for (int32 BoneIndex = 0; BoneIndex < Attributes.GetNumBones(); BoneIndex++)
+	{
+		const int32 ParentBoneIndex = ParentBoneIndexAttribute.Get(BoneIndex);
+		FTransform RefPoseTransform = InComponentSpaceTransforms[BoneIndex];
+
+		if (ParentBoneIndex != INDEX_NONE)
+		{
+			RefPoseTransform = RefPoseTransform.GetRelativeTransform(InComponentSpaceTransforms[ParentBoneIndex]);
+		}
+		WriteBonePoseAttribute.Set(BoneIndex, RefPoseTransform);
+	}
+	
+	return true;
+}
+
+
+bool FSkeletalMeshOperations::GetPosedMesh(
+	const FMeshDescription& InSourceMesh,
+	FMeshDescription& OutTargetMesh,
+	const TMap<FName, FTransform>& InBoneSpaceTransforms,
+	const FName InSkinWeightProfile, 
+	const TMap<FName, float>& InMorphTargetWeights
+	)
+{
+	// Verify that the mesh is valid.
+	FSkeletalMeshConstAttributes Attributes(InSourceMesh);
+	if (!Attributes.HasBoneNameAttribute() || !Attributes.HasBonePoseAttribute() || !Attributes.HasBoneParentIndexAttribute())
+	{
+		return false;
+	}
+
+	if (!Attributes.GetVertexSkinWeights(InSkinWeightProfile).IsValid())
+	{
+		return false;
+	}
+
+	TArray<FMatrix44f> RefToUserTransforms;
+	TArray<FMatrix44f> RefPoseTransforms;
+	TArray<FMatrix44f> UserPoseTransforms;
+	
+	FSkeletalMeshAttributes::FBoneNameAttributesConstRef BoneNameAttribute = Attributes.GetBoneNames(); 
+	FSkeletalMeshAttributes::FBonePoseAttributesConstRef BonePoseAttribute = Attributes.GetBonePoses(); 
+	FSkeletalMeshAttributes::FBoneParentIndexAttributesConstRef ParentBoneIndexAttribute = Attributes.GetBoneParentIndices();
+	
+	RefToUserTransforms.SetNumUninitialized(Attributes.GetNumBones());
+	RefPoseTransforms.SetNumUninitialized(Attributes.GetNumBones());
+	UserPoseTransforms.SetNumUninitialized(Attributes.GetNumBones());
+
+	for (int32 BoneIndex = 0; BoneIndex < Attributes.GetNumBones(); BoneIndex++)
+	{
+		const FName BoneName = BoneNameAttribute.Get(BoneIndex);
+		const int32 ParentBoneIndex = ParentBoneIndexAttribute.Get(BoneIndex);
+		RefPoseTransforms[BoneIndex] = FMatrix44f{BonePoseAttribute.Get(BoneIndex).ToMatrixWithScale()};
+		if (const FTransform* UserTransform = InBoneSpaceTransforms.Find(BoneName))
+		{
+			UserPoseTransforms[BoneIndex] = FMatrix44f{UserTransform->ToMatrixWithScale()};
+
+			// Update the pose on the mesh to match the user pose.
+		}
+		else
+		{
+			UserPoseTransforms[BoneIndex] = RefPoseTransforms[BoneIndex];
+		}
+		
+		if (ParentBoneIndex != INDEX_NONE)
+		{
+			RefPoseTransforms[BoneIndex] = RefPoseTransforms[BoneIndex] * RefPoseTransforms[ParentBoneIndex];
+			UserPoseTransforms[BoneIndex] = UserPoseTransforms[BoneIndex] * UserPoseTransforms[ParentBoneIndex];
+		}
+
+		RefToUserTransforms[BoneIndex] = RefPoseTransforms[BoneIndex].Inverse() * UserPoseTransforms[BoneIndex];
+	}
+	
+	Impl::GetPosedMesh(InSourceMesh, OutTargetMesh, RefToUserTransforms, InSkinWeightProfile, InMorphTargetWeights);
+
+	FSkeletalMeshAttributes WriteAttributes(OutTargetMesh);
+	FSkeletalMeshAttributes::FBonePoseAttributesRef WriteBonePoseAttribute = WriteAttributes.GetBonePoses(); 
+	for (int32 BoneIndex = 0; BoneIndex < Attributes.GetNumBones(); BoneIndex++)
+	{
+		const FName BoneName = BoneNameAttribute.Get(BoneIndex);
+		if (const FTransform* UserTransform = InBoneSpaceTransforms.Find(BoneName))
+		{
+			WriteBonePoseAttribute.Set(BoneIndex, *UserTransform);
+		}
+	}
+
+	return true;
+}
+
 
 #undef LOCTEXT_NAMESPACE
diff --git a/Engine/Source/Runtime/SkeletalMeshDescription/Public/SkeletalMeshOperations.h b/Engine/Source/Runtime/SkeletalMeshDescription/Public/SkeletalMeshOperations.h
--- a/Engine/Source/Runtime/SkeletalMeshDescription/Public/SkeletalMeshOperations.h
+++ b/Engine/Source/Runtime/SkeletalMeshDescription/Public/SkeletalMeshOperations.h
@@ -33,14 +33,14 @@ public:
 	
 	static SKELETALMESHDESCRIPTION_API void AppendSkinWeight(const FMeshDescription& SourceMesh, FMeshDescription& TargetMesh, FSkeletalMeshAppendSettings& AppendSettings);
 
-	/* Copies skin weight attribute from one mesh to another. Assumes the two geometries are identical or near-identical.
-	 * Uses closest triangle on the target mesh to interpolate skin weights to each of the points on the target mesh.
-	 * Attributes for the given profiles on both meshes should exist in order for this function to succeed. 
-	 * \param InSourceMesh The mesh to copy skin weights from.
-	 * \param InTargetMesh The mesh to copy skin weights to.
-	 * \param InSourceProfile The name of the skin weight profile on the source mesh to read from.
-	 * \param InTargetProfile The name of the skin weight profile on the target mesh to write to.
-	 * \param SourceBoneIndexToTargetBoneIndexMap An optional mapping table to map bone indexes on the source mesh to the target mesh.
+	/** Copies skin weight attribute from one mesh to another. Assumes the two geometries are identical or near-identical.
+	 *  Uses closest triangle on the target mesh to interpolate skin weights to each of the points on the target mesh.
+	 *  Attributes for the given profiles on both meshes should exist in order for this function to succeed. 
+	 *  @param InSourceMesh The mesh to copy skin weights from.
+	 *  @param InTargetMesh The mesh to copy skin weights to.
+	 *  @param InSourceProfile The name of the skin weight profile on the source mesh to read from.
+	 *  @param InTargetProfile The name of the skin weight profile on the target mesh to write to.
+	 *  @param SourceBoneIndexToTargetBoneIndexMap An optional mapping table to map bone indexes on the source mesh to the target mesh.
 	 *     The table needs to be complete for all the source bone indexes to valid target bone indexes, otherwise the behavior
 	 *     is undefined. If the table is not given, the bone indexes on the source and target meshes are assumed to be the same.
 	 */
@@ -55,14 +55,53 @@ public:
 	/** Remaps the bone indices on all skin weight attributes from one index to another. The array view should contain a full mapping of all the
 	 *  bone indices contained in the skin weights. The array is indexed by the current bone index and the value at that index is the new bone index.
 	 *  If the mapping is incomplete or if two entries map to the same bone, the result is undefined. No prior checking is performed.
-	 *  \param InMesh The mesh on which to modify all skin weight attributes to remap their bone indices.
-	 *  \param InBoneIndexMapping The mapping from one bone index to another. The old bone index is used to index into the array, the value at that
+	 *  @param InMesh The mesh on which to modify all skin weight attributes to remap their bone indices.
+	 *  @param InBoneIndexMapping The mapping from one bone index to another. The old bone index is used to index into the array, the value at that
 	 *     position is the new bone index.
-	 *  \return true if the operation was successful. If the mapping array was incomplete then this will return false. If there are no skin weight
+	 *  @return true if the operation was successful. If the mapping array was incomplete then this will return false. If there are no skin weight
 	 *     attributes on the mesh, then the operation is still deemed successful.
 	 */  
 	static SKELETALMESHDESCRIPTION_API bool RemapBoneIndicesOnSkinWeightAttribute(
 		FMeshDescription& InMesh,
 		TConstArrayView<int32> InBoneIndexMapping
 		);
+
+	/** Returns a mesh in the pose given by the component-space transforms passed in. The list of transforms should match exactly
+	 *  the list of bones stored on the mesh. If not, the function fails and returns \c false.
+	 *  If there are no skin weights on the mesh, or the named skin weight profile doesn't exist, the function also returns \c false.
+	 *  The resulting bones on the mesh will have their bone-space transforms updated so that the same mesh can be re-posed as needed.
+	 *  @param InSourceMesh The mesh to deform.
+	 *  @param OutTargetMesh The deformed mesh result.
+	 *  @param InComponentSpaceTransforms The component space transforms used to move the mesh joints for deforming, using linear-blend skinning.
+	 *  @param InSkinWeightProfile The skin weight profile to use as the source of skin weights for the deformation.
+	 *  @param InMorphTargetWeights Optional morph target weights to apply. Any morph target that doesn't exist is ignored.
+	 *  @return \c true if the operation succeeded.
+	 */
+	static SKELETALMESHDESCRIPTION_API bool GetPosedMesh(
+		const FMeshDescription& InSourceMesh,
+		FMeshDescription& OutTargetMesh,
+		TConstArrayView<FTransform> InComponentSpaceTransforms,
+		const FName InSkinWeightProfile = NAME_None,
+		const TMap<FName, float>& InMorphTargetWeights = {}
+		);
+
+	/** Returns a mesh in the pose given by the bone-space transforms passed in. The transforms simply replace the matching ref pose transforms
+	 *  stored in the bone data on the mesh. Any named transform, that does not match a bone on the mesh, is ignored.  
+	 *  If there are no skin weights on the mesh, or the named skin weight profile doesn't exist, the function also returns \c false.
+	 *  The resulting bones on the mesh will have their bone-space transforms updated so that the same mesh can be re-posed as needed.
+	 *  
+	 *  @param InSourceMesh The mesh to deform.
+	 *  @param OutTargetMesh The deformed mesh result.
+	 *  @param InBoneSpaceTransforms A map of named bone-space transforms.
+	 *  @param InSkinWeightProfile The skin weight profile to use as the source of skin weights for the deformation.
+	 *  @param InMorphTargetWeights Optional morph target weights to apply. Any morph target that doesn't exist is ignored.
+	 *  @return \c true if the operation succeeded.
+	 */
+	static SKELETALMESHDESCRIPTION_API bool GetPosedMesh(
+		const FMeshDescription& InSourceMesh,
+		FMeshDescription& OutTargetMesh,
+		const TMap<FName, FTransform>& InBoneSpaceTransforms, 
+		const FName InSkinWeightProfile = NAME_None,
+		const TMap<FName, float>& InMorphTargetWeights = {}
+		);
 };
