Mesh Description: Build fix for Mac.

#rnx

diff --git a/Engine/Source/Runtime/SkeletalMeshDescription/Private/SkeletalMeshOperations.cpp b/Engine/Source/Runtime/SkeletalMeshDescription/Private/SkeletalMeshOperations.cpp
--- a/Engine/Source/Runtime/SkeletalMeshDescription/Private/SkeletalMeshOperations.cpp
+++ b/Engine/Source/Runtime/SkeletalMeshDescription/Private/SkeletalMeshOperations.cpp
@@ -282,7 +282,7 @@ static void GetPosedMesh(
 	{
 		TVertexAttributesRef<FVector3f> PositionDelta;
 		TVertexInstanceAttributesRef<FVector3f> NormalDelta;
-		const float Weight = 0.0f;
+		float Weight = 0.0f;
 	};
 	
 	OutTargetMesh = InSourceMesh;
@@ -315,7 +315,11 @@ static void GetPosedMesh(
 			{
 				bAllMorphNormalsValid = false;
 			}
-			MorphInfos.Emplace(PositionDelta, NormalDelta, MorphWeight);
+			FMorphInfo MorphInfo;
+			MorphInfo.PositionDelta = PositionDelta;
+			MorphInfo.NormalDelta = NormalDelta;
+			MorphInfo.Weight = MorphWeight;
+			MorphInfos.Add(MorphInfo);
 		}
 	}
 
