Mesh Description: Ensure we have triangle normals to properly perform the normal/tangent re-computation after computing morph targets.

#rnx

diff --git a/Engine/Source/Runtime/SkeletalMeshDescription/Private/SkeletalMeshOperations.cpp b/Engine/Source/Runtime/SkeletalMeshDescription/Private/SkeletalMeshOperations.cpp
--- a/Engine/Source/Runtime/SkeletalMeshDescription/Private/SkeletalMeshOperations.cpp
+++ b/Engine/Source/Runtime/SkeletalMeshDescription/Private/SkeletalMeshOperations.cpp
@@ -268,11 +268,10 @@ bool FSkeletalMeshOperations::RemapBoneIndicesOnSkinWeightAttribute(FMeshDescrip
 	return true;
 }
 
-namespace Impl
+namespace UE::Impl
 {
-static void GetPosedMesh(
-	const FMeshDescription& InSourceMesh, 
-	FMeshDescription& OutTargetMesh,
+static void PoseMesh(
+	FMeshDescription& InOutTargetMesh,
 	TConstArrayView<FMatrix44f> InRefToUserTransforms, 
 	const FName InSkinWeightProfile, 
 	const TMap<FName, float>& InMorphTargetWeights
@@ -285,12 +284,11 @@ static void GetPosedMesh(
 		float Weight = 0.0f;
 	};
 	
-	OutTargetMesh = InSourceMesh;
-	FSkeletalMeshAttributes Attributes(OutTargetMesh);
+	FSkeletalMeshAttributes Attributes(InOutTargetMesh);
 
 	// We need the mesh to be compact for the parallel for to work.
 	FElementIDRemappings Remappings;
-	OutTargetMesh.Compact(Remappings);
+	InOutTargetMesh.Compact(Remappings);
 
 	TVertexAttributesRef<FVector3f> PositionAttribute = Attributes.GetVertexPositions();
 	TVertexInstanceAttributesRef<FVector3f> NormalAttribute = Attributes.GetVertexInstanceNormals();
@@ -332,8 +330,8 @@ static void GetPosedMesh(
 			TArray<FVertexID> Neighbors;
 		};
 		TArray<FMorphProcessContext> Contexts;
-		ParallelForWithTaskContext(Contexts, OutTargetMesh.Vertices().Num(),
-			[&Mesh=OutTargetMesh, &PositionAttribute, &MorphInfos, bAllMorphNormalsValid](FMorphProcessContext& Context, int32 Index)
+		ParallelForWithTaskContext(Contexts, InOutTargetMesh.Vertices().Num(),
+			[&Mesh=InOutTargetMesh, &PositionAttribute, &MorphInfos, bAllMorphNormalsValid](FMorphProcessContext& Context, int32 Index)
 			{
 				const FVertexID VertexID{Index};
 
@@ -370,7 +368,7 @@ static void GetPosedMesh(
 
 		if (bAllMorphNormalsValid)
 		{
-			ParallelForWithTaskContext(Contexts, OutTargetMesh.VertexInstances().Num(),
+			ParallelForWithTaskContext(Contexts, InOutTargetMesh.VertexInstances().Num(),
 				[&MorphInfos, &NormalAttribute, &TangentAttribute, &BinormalSignsAttribute](FMorphProcessContext& Context, int32 Index)
 				{
 					FVertexInstanceID VertexInstanceID{Index};
@@ -426,8 +424,15 @@ static void GetPosedMesh(
 				NormalAttribute.Set(VertexInstanceID, FVector3f::ZeroVector);
 			}
 
+			FSkeletalMeshOperations::ComputeTriangleTangentsAndNormals(InOutTargetMesh, UE_SMALL_NUMBER, nullptr);
+
 			// Compute the normals on the dirty vertices, and adjust the tangents to match.
-			FSkeletalMeshOperations::ComputeTangentsAndNormals(OutTargetMesh, EComputeNTBsFlags::WeightedNTBs);
+			FSkeletalMeshOperations::ComputeTangentsAndNormals(InOutTargetMesh, EComputeNTBsFlags::WeightedNTBs);
+			
+			// We don't need the triangle tangents and normals anymore.
+			InOutTargetMesh.TriangleAttributes().UnregisterAttribute(MeshAttribute::Triangle::Normal);
+			InOutTargetMesh.TriangleAttributes().UnregisterAttribute(MeshAttribute::Triangle::Tangent);
+			InOutTargetMesh.TriangleAttributes().UnregisterAttribute(MeshAttribute::Triangle::Binormal);
 		}
 	}
 
@@ -443,8 +448,8 @@ static void GetPosedMesh(
 	}
 	
 	FSkinWeightsVertexAttributesRef SkinWeightAttribute = Attributes.GetVertexSkinWeights(InSkinWeightProfile);
-	ParallelFor(OutTargetMesh.Vertices().Num(),
-		[&Mesh=OutTargetMesh, &PositionAttribute, &NormalAttribute, &TangentAttribute, &SkinWeightAttribute, &RefToUserTransforms=InRefToUserTransforms, &RefToUserTransformsNormal](int32 Index)
+	ParallelFor(InOutTargetMesh.Vertices().Num(),
+		[&Mesh=InOutTargetMesh, &PositionAttribute, &NormalAttribute, &TangentAttribute, &SkinWeightAttribute, &RefToUserTransforms=InRefToUserTransforms, &RefToUserTransformsNormal](int32 Index)
 		{
 			const FVertexID VertexID(Index);
 			const FVertexBoneWeights BoneWeights = SkinWeightAttribute.Get(VertexID);
@@ -531,7 +536,9 @@ bool FSkeletalMeshOperations::GetPosedMesh(
 		RefToUserTransforms[BoneIndex] = RefPoseTransforms[BoneIndex].Inverse() * FMatrix44f{InComponentSpaceTransforms[BoneIndex].ToMatrixWithScale()};
 	}
 	
-	Impl::GetPosedMesh(InSourceMesh, OutTargetMesh, RefToUserTransforms, InSkinWeightProfile, InMorphTargetWeights);
+	// Start with a fresh duplicate and then pose the target mesh in-place.
+	OutTargetMesh = InSourceMesh;
+	UE::Impl::PoseMesh(OutTargetMesh, RefToUserTransforms, InSkinWeightProfile, InMorphTargetWeights);
 
 	// Write out the current ref pose (in bone-space) to the mesh. 
 	FSkeletalMeshAttributes WriteAttributes(OutTargetMesh);
@@ -608,8 +615,10 @@ bool FSkeletalMeshOperations::GetPosedMesh(
 
 		RefToUserTransforms[BoneIndex] = RefPoseTransforms[BoneIndex].Inverse() * UserPoseTransforms[BoneIndex];
 	}
-	
-	Impl::GetPosedMesh(InSourceMesh, OutTargetMesh, RefToUserTransforms, InSkinWeightProfile, InMorphTargetWeights);
+
+	// Start with a fresh duplicate and then pose the target mesh in-place.
+	OutTargetMesh = InSourceMesh;
+	UE::Impl::PoseMesh(OutTargetMesh, RefToUserTransforms, InSkinWeightProfile, InMorphTargetWeights);
 
 	FSkeletalMeshAttributes WriteAttributes(OutTargetMesh);
 	FSkeletalMeshAttributes::FBonePoseAttributesRef WriteBonePoseAttribute = WriteAttributes.GetBonePoses(); 
