[Interchange] SkeletalMeshDescription build health fix.

#rnx

diff --git a/Engine/Plugins/Interchange/Runtime/Source/Import/Private/Mesh/InterchangeSkeletalMeshFactory.cpp b/Engine/Plugins/Interchange/Runtime/Source/Import/Private/Mesh/InterchangeSkeletalMeshFactory.cpp
--- a/Engine/Plugins/Interchange/Runtime/Source/Import/Private/Mesh/InterchangeSkeletalMeshFactory.cpp
+++ b/Engine/Plugins/Interchange/Runtime/Source/Import/Private/Mesh/InterchangeSkeletalMeshFactory.cpp
@@ -571,7 +571,9 @@ namespace UE
 
 				FSkeletalMeshOperations::ConvertHardEdgesToSmoothMasks(DestinationMeshDescription, FaceSmoothingMasks); //Create Smooth Masks from original mesh description
 				FSkeletalMeshOperations::FixVertexInstanceStructure(DestinationMeshDescription, FixedMeshDescription, FaceSmoothingMasks, FixedFaceSmoothingMasks);
-				FSkeletalMeshOperations::ValidateFixComputeMeshDescriptionData(SkeletalMesh, FixedMeshDescription, FixedFaceSmoothingMasks, CurrentLodIndex); //Validates and Fixes MeshDescription including re-instating Smooth masks.
+				FString SkeletalMeshPath = SkeletalMesh ? SkeletalMesh->GetOuter()->GetPathName() : TEXT("");
+				bool bComputeWeightedNormals = SkeletalMesh ? SkeletalMesh->GetLODInfo(CurrentLodIndex) && SkeletalMesh->GetLODInfo(CurrentLodIndex)->BuildSettings.bComputeWeightedNormals : false;
+				FSkeletalMeshOperations::ValidateFixComputeMeshDescriptionData(FixedMeshDescription, FixedFaceSmoothingMasks, CurrentLodIndex, bComputeWeightedNormals, SkeletalMeshPath); //Validates and Fixes MeshDescription including re-instating Smooth masks.
 
 				DestinationMeshDescription = MoveTemp(FixedMeshDescription);
 			}
diff --git a/Engine/Source/Runtime/SkeletalMeshDescription/Private/SkeletalMeshOperations.cpp b/Engine/Source/Runtime/SkeletalMeshDescription/Private/SkeletalMeshOperations.cpp
--- a/Engine/Source/Runtime/SkeletalMeshDescription/Private/SkeletalMeshOperations.cpp
+++ b/Engine/Source/Runtime/SkeletalMeshDescription/Private/SkeletalMeshOperations.cpp
@@ -3,7 +3,6 @@
 #include "SkeletalMeshOperations.h"
 
 #include "BoneWeights.h"
-#include "Engine/SkeletalMesh.h"
 #include "MeshDescriptionAdapter.h"
 #include "Math/GenericOctree.h"
 #include "SkeletalMeshAttributes.h"
@@ -900,7 +899,7 @@ void FSkeletalMeshOperations::FixVertexInstanceStructure(FMeshDescription& Sourc
 	TargetMeshDescription.ResumeUVIndexing();
 };
 
-void FSkeletalMeshOperations::ValidateFixComputeMeshDescriptionData(USkeletalMesh* SkeletalMesh, FMeshDescription& MeshDescription, const TArray<uint32>& FaceSmoothingMasks, int32 LODIndex)
+void FSkeletalMeshOperations::ValidateFixComputeMeshDescriptionData(FMeshDescription& MeshDescription, const TArray<uint32>& FaceSmoothingMasks, int32 LODIndex, const bool bComputeWeightedNormals, const FString& SkeletalMeshPath)
 {
 	//As FSkeletalMeshImportData::GetMeshDescription(...) (which we replaced with direct FMeshDescription usage) is doing the following calls:
 	// - ConvertSmoothGroupToHardEdges
@@ -913,12 +912,6 @@ void FSkeletalMeshOperations::ValidateFixComputeMeshDescriptionData(USkeletalMes
 		//ConvertHardEdgesToSmoothMasks(MeshDescription, FaceSmoothingMasks);
 		FSkeletalMeshOperations::ConvertSmoothGroupToHardEdges(FaceSmoothingMasks, MeshDescription);
 
-		FString SkeletalMeshPath;
-		if (SkeletalMesh)
-		{
-			SkeletalMeshPath = SkeletalMesh->GetOuter()->GetPathName();
-		}
-
 		// Check if we have any broken data, including UVs and normals/tangents.
 		FSkeletalMeshOperations::ValidateAndFixData(MeshDescription, SkeletalMeshPath);
 
@@ -930,7 +923,7 @@ void FSkeletalMeshOperations::ValidateFixComputeMeshDescriptionData(USkeletalMes
 			FSkeletalMeshOperations::ComputeTriangleTangentsAndNormals(MeshDescription, UE_SMALL_NUMBER, !SkeletalMeshPath.IsEmpty() ? *SkeletalMeshPath : nullptr);
 
 			EComputeNTBsFlags ComputeNTBsOptions = EComputeNTBsFlags::None;
-			if (SkeletalMesh->GetLODInfo(LODIndex) && SkeletalMesh->GetLODInfo(LODIndex)->BuildSettings.bComputeWeightedNormals)
+			if (bComputeWeightedNormals)
 			{
 				ComputeNTBsOptions |= EComputeNTBsFlags::WeightedNTBs;
 			}
diff --git a/Engine/Source/Runtime/SkeletalMeshDescription/Public/SkeletalMeshOperations.h b/Engine/Source/Runtime/SkeletalMeshDescription/Public/SkeletalMeshOperations.h
--- a/Engine/Source/Runtime/SkeletalMeshDescription/Public/SkeletalMeshOperations.h
+++ b/Engine/Source/Runtime/SkeletalMeshDescription/Public/SkeletalMeshOperations.h
@@ -10,7 +10,6 @@
 #include "MeshTypes.h"
 #include "StaticMeshOperations.h"
 
-class USkeletalMesh;
 struct FMeshDescription;
 
 
@@ -131,7 +130,7 @@ public:
 	 *		- HasInvalidVertexInstanceNormalsOrTangents -> ComputeTriangleTangentsAndNormals
 	 *		- (Re)BuildIndexers
 	 */ 
-	static SKELETALMESHDESCRIPTION_API void ValidateFixComputeMeshDescriptionData(USkeletalMesh* SkeletalMesh, FMeshDescription& MeshDescription, const TArray<uint32>& FaceSmoothingMasks, int32 LODIndex);
+	static SKELETALMESHDESCRIPTION_API void ValidateFixComputeMeshDescriptionData(FMeshDescription& MeshDescription, const TArray<uint32>& FaceSmoothingMasks, int32 LODIndex, const bool bComputeWeightedNormals, const FString& SkeletalMeshPath);
 
 	/**
 	 * Function will do the following steps on the MeshDescription (SkeletalMeshDescription) :
