[Interchange] Uploading GeometryPart data to SkeletalMeshAttributes.

#rnx

diff --git a/Engine/Plugins/Interchange/Runtime/Source/Import/Private/Mesh/InterchangeSkeletalMeshFactory.cpp b/Engine/Plugins/Interchange/Runtime/Source/Import/Private/Mesh/InterchangeSkeletalMeshFactory.cpp
--- a/Engine/Plugins/Interchange/Runtime/Source/Import/Private/Mesh/InterchangeSkeletalMeshFactory.cpp
+++ b/Engine/Plugins/Interchange/Runtime/Source/Import/Private/Mesh/InterchangeSkeletalMeshFactory.cpp
@@ -378,6 +378,12 @@ namespace UE
 
 				FSkeletalMeshAttributes DestinationMeshAttributes(DestinationMeshDescription);
 				DestinationMeshAttributes.Register();
+
+				DestinationMeshAttributes.RegisterSourceGeometryPartsAttributes();
+				FSkeletalMeshAttributes::FSourceGeometryPartNameRef NameAttribute = DestinationMeshAttributes.GetSourceGeometryPartNames();
+				FSkeletalMeshAttributes::FSourceGeometryPartVertexOffsetAndCountRef VertexAndCountAttribute = DestinationMeshAttributes.GetSourceGeometryPartVertexOffsetAndCounts();
+				DestinationMeshAttributes.SourceGeometryParts().Reserve(MeshReferences.Num());
+
 				FStaticMeshOperations::FAppendSettings AppendSettings;
 				for (int32 ChannelIdx = 0; ChannelIdx < FStaticMeshOperations::FAppendSettings::MAX_NUM_UV_CHANNELS; ++ChannelIdx)
 				{
@@ -443,6 +449,15 @@ namespace UE
 					
 					const int32 VertexOffset = DestinationMeshDescription.Vertices().Num();
 
+					if (MeshNodeContext.SceneNode && MeshNodeContext.MeshNode)
+					{
+						FSourceGeometryPartID PartID = DestinationMeshAttributes.CreateSourceGeometryPart();
+
+						NameAttribute.Set(PartID, FName(MeshNodeContext.SceneNode->GetDisplayLabel()));
+						VertexAndCountAttribute.Set(PartID, { DestinationMeshDescription.Vertices().Num(), LodMeshPayload->MeshDescription.Vertices().Num()});
+					}
+					
+
 					FSkeletalMeshOperations::FSkeletalMeshAppendSettings SkeletalMeshAppendSettings;
 					SkeletalMeshAppendSettings.bAppendVertexAttributes = bImportVertexAttributes;
 					SkeletalMeshAppendSettings.SourceVertexIDOffset = VertexOffset;
diff --git a/Engine/Source/Runtime/SkeletalMeshDescription/Private/SkeletalMeshOperations.cpp b/Engine/Source/Runtime/SkeletalMeshDescription/Private/SkeletalMeshOperations.cpp
--- a/Engine/Source/Runtime/SkeletalMeshDescription/Private/SkeletalMeshOperations.cpp
+++ b/Engine/Source/Runtime/SkeletalMeshDescription/Private/SkeletalMeshOperations.cpp
@@ -718,6 +718,10 @@ void FSkeletalMeshOperations::FixVertexInstanceStructure(FMeshDescription& Sourc
 		TVertexInstanceAttributesRef<FVector4f> Colors;
 		TVertexInstanceAttributesRef<FVector2f> UVs;
 
+		//GeometryParts
+		FSkeletalMeshAttributes::FSourceGeometryPartNameRef GeometryPartNames;
+		FSkeletalMeshAttributes::FSourceGeometryPartVertexOffsetAndCountRef GeometryPartVertexOffsetAndCounts;
+
 		//For read only
 		int32 NumUVChannels;
 		int32 NumberOfMorphTargets;
@@ -731,6 +735,8 @@ void FSkeletalMeshOperations::FixVertexInstanceStructure(FMeshDescription& Sourc
 			if (bRegisterAttributes)
 			{
 				MeshAttributes.Register();
+
+				MeshAttributes.RegisterSourceGeometryPartsAttributes();
 			}
 
 			PolygonGroupMaterialSlotNames = MeshAttributes.GetPolygonGroupMaterialSlotNames();
@@ -744,6 +750,9 @@ void FSkeletalMeshOperations::FixVertexInstanceStructure(FMeshDescription& Sourc
 			Colors = MeshAttributes.GetVertexInstanceColors();
 			UVs = MeshAttributes.GetVertexInstanceUVs();
 
+			GeometryPartNames = MeshAttributes.GetSourceGeometryPartNames();
+			GeometryPartVertexOffsetAndCounts = MeshAttributes.GetSourceGeometryPartVertexOffsetAndCounts();
+
 			NumUVChannels = UVs.GetNumChannels();
 
 			MorphTargetNames = MeshAttributes.GetMorphTargetNames();
@@ -892,6 +901,18 @@ void FSkeletalMeshOperations::FixVertexInstanceStructure(FMeshDescription& Sourc
 		}
 	}
 
+	//Copy GeometryParts:
+	for (const FSourceGeometryPartID GeometryPartID : Source.MeshAttributes.SourceGeometryParts().GetElementIDs())
+	{
+		FName Name = Source.GeometryPartNames.Get(GeometryPartID);
+		TArrayView<const int32> OffsetAndCount = Source.GeometryPartVertexOffsetAndCounts.Get(GeometryPartID);
+
+		FSourceGeometryPartID PartID = Target.MeshAttributes.CreateSourceGeometryPart();
+
+		Target.GeometryPartNames.Set(PartID, Name);
+		Target.GeometryPartVertexOffsetAndCounts.Set(PartID, { OffsetAndCount[0], OffsetAndCount[1]});
+	}
+
 	TargetMeshDescription.ResumeVertexInstanceIndexing();
 	TargetMeshDescription.ResumeEdgeIndexing();
 	TargetMeshDescription.ResumePolygonIndexing();
