[SkeletalMeshOperation] avoid unnecessary compact + added in-place variants of the posing APIs to avoid mandatory copy

diff --git a/Engine/Source/Runtime/SkeletalMeshDescription/Private/SkeletalMeshOperations.cpp b/Engine/Source/Runtime/SkeletalMeshDescription/Private/SkeletalMeshOperations.cpp
--- a/Engine/Source/Runtime/SkeletalMeshDescription/Private/SkeletalMeshOperations.cpp
+++ b/Engine/Source/Runtime/SkeletalMeshDescription/Private/SkeletalMeshOperations.cpp
@@ -288,8 +288,16 @@ static void PoseMesh(
 	FSkeletalMeshAttributes Attributes(InOutTargetMesh);
 
 	// We need the mesh to be compact for the parallel for to work.
-	FElementIDRemappings Remappings;
-	InOutTargetMesh.Compact(Remappings);
+	if (InOutTargetMesh.NeedsCompact())
+	{
+		FElementIDRemappings Remappings;
+		InOutTargetMesh.Compact(Remappings);
+	}
+	else
+	{
+		// Make sure indexers are built before entering parallel work
+		InOutTargetMesh.BuildVertexIndexers();
+	}
 
 	TVertexAttributesRef<FVector3f> PositionAttribute = Attributes.GetVertexPositions();
 	TVertexInstanceAttributesRef<FVector3f> NormalAttribute = Attributes.GetVertexInstanceNormals();
@@ -485,7 +493,7 @@ static void PoseMesh(
 		});
 }
 
-static void GetMorphTargetDeltasFromPosedMesh(
+static void UnposeMesh(
 	FMeshDescription& InOutTargetMesh,
 	const FMeshDescription& InRefMesh,
 	TConstArrayView<FMatrix44f> InRefToUserTransforms,
@@ -504,8 +512,16 @@ static void GetMorphTargetDeltasFromPosedMesh(
 	FSkeletalMeshConstAttributes RefAttributes(InRefMesh);
 
 	// We need the mesh to be compact for the parallel for to work.
-	FElementIDRemappings Remappings;
-	InOutTargetMesh.Compact(Remappings);
+	if (InOutTargetMesh.NeedsCompact())
+	{
+		FElementIDRemappings Remappings;
+		InOutTargetMesh.Compact(Remappings);
+	}
+	else
+	{
+		// Make sure indexers are built before entering parallel work
+		InOutTargetMesh.BuildVertexIndexers();
+	}
 
 	TVertexAttributesRef<FVector3f> PositionAttribute = Attributes.GetVertexPositions();
 	TVertexInstanceAttributesRef<FVector3f> NormalAttribute = Attributes.GetVertexInstanceNormals();
@@ -715,6 +731,76 @@ bool FSkeletalMeshOperations::GetPosedMesh(
 	return true;
 }
 
+bool FSkeletalMeshOperations::GetPosedMeshInPlace(FMeshDescription& InOutTargetMesh, TConstArrayView<FTransform> InComponentSpaceTransforms,
+	const FName InSkinWeightProfile, const TMap<FName, float>& InMorphTargetWeights, bool bInWriteBonePose)
+{
+	// Verify that the mesh is valid.
+	FSkeletalMeshConstAttributes Attributes(InOutTargetMesh);
+	if (!Attributes.HasBonePoseAttribute() || !Attributes.HasBoneParentIndexAttribute())
+	{
+		return false;
+	}
+
+	if (!Attributes.GetVertexSkinWeights(InSkinWeightProfile).IsValid())
+	{
+		return false;
+	}
+
+	if (Attributes.GetNumBones() != InComponentSpaceTransforms.Num())
+	{
+		return false;
+	}
+
+	// Convert the component-space transforms into a set of matrices that transform from the reference pose to
+	// the user pose. These are then used to nudge the vertices from the reference pose to the wanted
+	// user pose by weighing the influence of each bone on a given vertex. If the user pose and the reference pose
+	// are identical, these are all identity matrices.
+	TArray<FMatrix44f> RefToUserTransforms;
+	TArray<FMatrix44f> RefPoseTransforms;
+	
+	FSkeletalMeshAttributes::FBonePoseAttributesConstRef BonePoseAttribute = Attributes.GetBonePoses(); 
+	FSkeletalMeshAttributes::FBoneParentIndexAttributesConstRef ParentBoneIndexAttribute = Attributes.GetBoneParentIndices(); 
+	RefToUserTransforms.SetNumUninitialized(Attributes.GetNumBones());
+	RefPoseTransforms.SetNumUninitialized(Attributes.GetNumBones());
+	
+	for (int32 BoneIndex = 0; BoneIndex < Attributes.GetNumBones(); BoneIndex++)
+	{
+		const int32 ParentBoneIndex = ParentBoneIndexAttribute.Get(BoneIndex);
+		RefPoseTransforms[BoneIndex] = FMatrix44f{BonePoseAttribute.Get(BoneIndex).ToMatrixWithScale()};
+
+		if (ParentBoneIndex != INDEX_NONE)
+		{
+			RefPoseTransforms[BoneIndex] = RefPoseTransforms[BoneIndex] * RefPoseTransforms[ParentBoneIndex];
+		}
+
+		RefToUserTransforms[BoneIndex] = RefPoseTransforms[BoneIndex].Inverse() * FMatrix44f{InComponentSpaceTransforms[BoneIndex].ToMatrixWithScale()};
+	}
+	
+	// Pose the target mesh in-place.
+	UE::Impl::PoseMesh(InOutTargetMesh, RefToUserTransforms, InSkinWeightProfile, InMorphTargetWeights);
+
+	if (bInWriteBonePose)
+	{
+		// Write out the current ref pose (in bone-space) to the mesh. 
+		FSkeletalMeshAttributes WriteAttributes(InOutTargetMesh);
+		FSkeletalMeshAttributes::FBonePoseAttributesRef WriteBonePoseAttribute = WriteAttributes.GetBonePoses(); 
+		for (int32 BoneIndex = 0; BoneIndex < Attributes.GetNumBones(); BoneIndex++)
+		{
+			const int32 ParentBoneIndex = ParentBoneIndexAttribute.Get(BoneIndex);
+			FTransform RefPoseTransform = InComponentSpaceTransforms[BoneIndex];
+
+			if (ParentBoneIndex != INDEX_NONE)
+			{
+				RefPoseTransform = RefPoseTransform.GetRelativeTransform(InComponentSpaceTransforms[ParentBoneIndex]);
+			}
+			WriteBonePoseAttribute.Set(BoneIndex, RefPoseTransform);
+		}
+	}
+
+	
+	return true;
+}
+
 
 bool FSkeletalMeshOperations::GetPosedMesh(
 	const FMeshDescription& InSourceMesh,
@@ -791,18 +877,18 @@ bool FSkeletalMeshOperations::GetPosedMesh(
 	return true;
 }
 
-bool FSkeletalMeshOperations::GetMorphTargetDeltasFromPosedMesh(
-	const FMeshDescription& InPoseMesh,
+bool FSkeletalMeshOperations::GetUnposedMesh(
+	const FMeshDescription& InPosedMesh,
 	const FMeshDescription& InRefMesh,
 	TArray<FTransform>& RefBoneTransforms,
-	FMeshDescription& OutTargetMesh,
+	FMeshDescription& OutUnposedMesh,
 	TConstArrayView<FTransform> InComponentSpaceTransforms,
 	const FName InSkinWeightProfile,
 	const TMap<FName, float>& InMorphTargetWeights
 )
 {
 	// Verify that the mesh is valid.
-	FSkeletalMeshConstAttributes Attributes(InPoseMesh);
+	FSkeletalMeshConstAttributes Attributes(InPosedMesh);
 	if (!Attributes.HasBonePoseAttribute() || !Attributes.HasBoneParentIndexAttribute())
 	{
 		return false;
@@ -840,11 +926,11 @@ bool FSkeletalMeshOperations::GetMorphTargetDeltasFromPosedMesh(
 	}
 
 	// Start with a fresh duplicate and then pose the target mesh in-place.
-	OutTargetMesh = InPoseMesh;
-	UE::Impl::GetMorphTargetDeltasFromPosedMesh(OutTargetMesh, InRefMesh, RefToUserTransforms, InSkinWeightProfile, InMorphTargetWeights);
+	OutUnposedMesh = InPosedMesh;
+	UE::Impl::UnposeMesh(OutUnposedMesh, InRefMesh, RefToUserTransforms, InSkinWeightProfile, InMorphTargetWeights);
 
 	// Write out the current ref pose (in bone-space) to the mesh. 
-	FSkeletalMeshAttributes WriteAttributes(OutTargetMesh);
+	FSkeletalMeshAttributes WriteAttributes(OutUnposedMesh);
 	FSkeletalMeshAttributes::FBonePoseAttributesRef WriteBonePoseAttribute = WriteAttributes.GetBonePoses();
 	for (int32 BoneIndex = 0; BoneIndex < Attributes.GetNumBones(); BoneIndex++)
 	{
@@ -861,6 +947,71 @@ bool FSkeletalMeshOperations::GetMorphTargetDeltasFromPosedMesh(
 	return true;
 }
 
+bool FSkeletalMeshOperations::GetUnposedMeshInPlace(FMeshDescription& InOutTargetMesh, const FMeshDescription& InRefMesh,
+	TArray<FTransform>& RefBoneTransforms, TConstArrayView<FTransform> InComponentSpaceTransforms, const FName InSkinWeightProfile,
+	const TMap<FName, float>& InMorphTargetWeights, bool bInWriteBonePose)
+{
+	// Verify that the mesh is valid.
+	FSkeletalMeshConstAttributes Attributes(InOutTargetMesh);
+	if (!Attributes.HasBonePoseAttribute() || !Attributes.HasBoneParentIndexAttribute())
+	{
+		return false;
+	}
+
+	if (!Attributes.GetVertexSkinWeights(InSkinWeightProfile).IsValid())
+	{
+		return false;
+	}
+
+	if (Attributes.GetNumBones() != InComponentSpaceTransforms.Num())
+	{
+		return false;
+	}
+
+	// Convert the component-space transforms into a set of matrices that transform from the reference pose to
+	// the user pose. These are then used to nudge the vertices from the reference pose to the wanted
+	// user pose by weighing the influence of each bone on a given vertex. If the user pose and the reference pose
+	// are identical, these are all identity matrices.
+	TArray<FMatrix44f> RefToUserTransforms;
+	TArray<FMatrix44f> RefPoseTransforms;
+
+	FSkeletalMeshAttributes::FBonePoseAttributesConstRef BonePoseAttribute = Attributes.GetBonePoses();
+	FSkeletalMeshAttributes::FBoneParentIndexAttributesConstRef ParentBoneIndexAttribute = Attributes.GetBoneParentIndices();
+	RefToUserTransforms.SetNumUninitialized(Attributes.GetNumBones());
+	RefPoseTransforms.SetNumUninitialized(Attributes.GetNumBones());
+
+	for (int32 BoneIndex = 0; BoneIndex < Attributes.GetNumBones(); BoneIndex++)
+	{
+		const int32 ParentBoneIndex = ParentBoneIndexAttribute.Get(BoneIndex);
+		
+		FTransform ReferenceBoneTransform = RefBoneTransforms[BoneIndex];
+		RefPoseTransforms[BoneIndex] = FMatrix44f{ ReferenceBoneTransform.ToMatrixWithScale() };
+		RefToUserTransforms[BoneIndex] = RefPoseTransforms[BoneIndex].Inverse() * FMatrix44f { InComponentSpaceTransforms[BoneIndex].ToMatrixWithScale() };
+	}
+
+	UE::Impl::UnposeMesh(InOutTargetMesh, InRefMesh, RefToUserTransforms, InSkinWeightProfile, InMorphTargetWeights);
+
+	if (bInWriteBonePose)
+	{
+		// Write out the current ref pose (in bone-space) to the mesh. 
+		FSkeletalMeshAttributes WriteAttributes(InOutTargetMesh);
+		FSkeletalMeshAttributes::FBonePoseAttributesRef WriteBonePoseAttribute = WriteAttributes.GetBonePoses();
+		for (int32 BoneIndex = 0; BoneIndex < Attributes.GetNumBones(); BoneIndex++)
+		{
+			const int32 ParentBoneIndex = ParentBoneIndexAttribute.Get(BoneIndex);
+			FTransform RefPoseTransform = InComponentSpaceTransforms[BoneIndex];
+
+			if (ParentBoneIndex != INDEX_NONE)
+			{
+				RefPoseTransform = RefPoseTransform.GetRelativeTransform(InComponentSpaceTransforms[ParentBoneIndex]);
+			}
+			WriteBonePoseAttribute.Set(BoneIndex, RefPoseTransform);
+		}	
+	}
+
+	return true;
+}
+
 void FSkeletalMeshOperations::ConvertHardEdgesToSmoothMasks(
 	const FMeshDescription& InMeshDescription,
 	TArray<uint32>& OutSmoothMasks
diff --git a/Engine/Source/Runtime/SkeletalMeshDescription/Public/SkeletalMeshOperations.h b/Engine/Source/Runtime/SkeletalMeshDescription/Public/SkeletalMeshOperations.h
--- a/Engine/Source/Runtime/SkeletalMeshDescription/Public/SkeletalMeshOperations.h
+++ b/Engine/Source/Runtime/SkeletalMeshDescription/Public/SkeletalMeshOperations.h
@@ -105,27 +105,66 @@ public:
 		const TMap<FName, float>& InMorphTargetWeights = {}
 		);
 
-	/** Returns a Morph Target deltas extracted from the deformed mesh
+	/** Returns a mesh in the pose given by the component-space transforms passed in. The list of transforms should match exactly
+	 *  the list of bones stored on the mesh. If not, the function fails and returns \c false.
+	 *  If there are no skin weights on the mesh, or the named skin weight profile doesn't exist, the function also returns \c false.
+	 *  if bInWriteBonePose is true, the resulting bones on the mesh will have their bone-space transforms updated so that the same mesh can be re-posed as needed.
+	 *  @param InOutTargetMesh rest pose mesh to be deformed in place
+	 *  @param InComponentSpaceTransforms The component space transforms used to move the mesh joints for deforming, using linear-blend skinning.
+	 *  @param InSkinWeightProfile The skin weight profile to use as the source of skin weights for the deformation.
+	 *  @param InMorphTargetWeights Optional morph target weights to apply. Any morph target that doesn't exist is ignored.
+	 *  @param bInWriteBonePose whether to write the current bone pose into the mesh description.
+	 *  @return \c true if the operation succeeded.
+	 */
+	static SKELETALMESHDESCRIPTION_API bool GetPosedMeshInPlace(
+		FMeshDescription& InOutTargetMesh,
+		TConstArrayView<FTransform> InComponentSpaceTransforms,
+		const FName InSkinWeightProfile = NAME_None,
+		const TMap<FName, float>& InMorphTargetWeights = {},
+		bool bInWriteBonePose = false
+		);
+
+	/** Returns the unposed version of the provided posed mesh
 	 *  
-	 *  @param InPoseMesh Deformed mesh
-	 *  @param InRefMesh Ref pose mesh
+	 *  @param InPosedMesh The posed mesh,
+	 *  @param InRefMesh Ref mesh containing morph target deltas of existing morphs
 	 *  @param RefBoneTransforms Ref component space bone transforms
-	 *  @param OutTargetMesh Extracted morph target deltas 
+	 *  @param OutUnposedMesh The mesh with pose reset to reference
 	 *  @param InComponentSpaceTransforms Current component space bone transforms
 	 *  @param InSkinWeightProfile The skin weight profile to use as the source of skin weights for the deformation.
 	 *  @param InMorphTargetWeights Active morph target weights producing the fully deformed mesh. Don't include the weight of the morph target you want to extract
 	 *  @return \c true if the operation succeeded.
 	 */
-	static SKELETALMESHDESCRIPTION_API bool GetMorphTargetDeltasFromPosedMesh(
-		const FMeshDescription& InPoseMesh,
+	static SKELETALMESHDESCRIPTION_API bool GetUnposedMesh(
+		const FMeshDescription& InPosedMesh,
 		const FMeshDescription& InRefMesh,
 		TArray<FTransform>& RefBoneTransforms,
-		FMeshDescription& OutTargetMesh,
+		FMeshDescription& OutUnposedMesh,
 		TConstArrayView<FTransform> InComponentSpaceTransforms,
 		const FName InSkinWeightProfile,
 		const TMap<FName, float>& InMorphTargetWeights
 		);
 
+	/** Returns the unposed version of the provided posed mesh
+	 *  
+	 *  @param InOutTargetMesh input is the posed mesh, output is the mesh reset to ref pose
+	 *  @param InRefMesh Ref mesh containing morph target deltas of existing morphs
+	 *  @param RefBoneTransforms Ref component space bone transforms
+	 *  @param InComponentSpaceTransforms Current component space bone transforms
+	 *  @param InSkinWeightProfile The skin weight profile to use as the source of skin weights for the deformation.
+	 *  @param InMorphTargetWeights Active morph target weights producing the fully deformed mesh. Don't include the weight of the morph target you want to extract
+	 *  @return \c true if the operation succeeded.
+	 */
+	static SKELETALMESHDESCRIPTION_API bool GetUnposedMeshInPlace(
+		FMeshDescription& InOutTargetMesh,
+		const FMeshDescription& InRefMesh,
+		TArray<FTransform>& RefBoneTransforms,
+		TConstArrayView<FTransform> InComponentSpaceTransforms,
+		const FName InSkinWeightProfile,
+		const TMap<FName, float>& InMorphTargetWeights,
+		bool bInWriteBonePose = false
+		);
+
 	/**
 	 * A simpler variant of FStaticMeshOperations::ConvertHardEdgesToSmoothGroup that assumes that hard edges always form closed regions. 
 	 */
