[SkeletalMeshOperations] Optimizations for FixVertexInstanceStructure.
With fixes in this cl the file in the ticket went from ~18 seconds to sub 3 seconds on FixVertexInstancesStructure function.

#rnx

diff --git a/Engine/Source/Runtime/SkeletalMeshDescription/Private/SkeletalMeshOperations.cpp b/Engine/Source/Runtime/SkeletalMeshDescription/Private/SkeletalMeshOperations.cpp
--- a/Engine/Source/Runtime/SkeletalMeshDescription/Private/SkeletalMeshOperations.cpp
+++ b/Engine/Source/Runtime/SkeletalMeshDescription/Private/SkeletalMeshOperations.cpp
@@ -1147,6 +1147,8 @@ void FSkeletalMeshOperations::ConvertHardEdgesToSmoothMasks(
 
 void FSkeletalMeshOperations::FixVertexInstanceStructure(FMeshDescription& SourceMeshDescription, FMeshDescription& TargetMeshDescription /*Expected to be empty*/, const TArray<uint32>& SourceSmoothingMasks, TArray<uint32>& TargetFaceSmoothingMasks)
 {
+	TRACE_CPUPROFILER_EVENT_SCOPE(FixVertexInstanceStructure)
+
 	if (!ensure(TargetMeshDescription.IsEmpty()))
 	{
 		return;
@@ -1223,38 +1225,19 @@ void FSkeletalMeshOperations::FixVertexInstanceStructure(FMeshDescription& Sourc
 	FMeshAttributesHelper Source(SourceMeshDescription);
 	FMeshAttributesHelper Target(TargetMeshDescription, true);
 
-	TargetMeshDescription.SuspendVertexInstanceIndexing();
-	TargetMeshDescription.SuspendEdgeIndexing();
-	TargetMeshDescription.SuspendPolygonIndexing();
-	TargetMeshDescription.SuspendPolygonGroupIndexing();
-	TargetMeshDescription.SuspendUVIndexing();
-
-	TargetMeshDescription.ReserveNewVertices(SourceMeshDescription.Vertices().Num());
-	TargetMeshDescription.ReserveNewTriangles(SourceMeshDescription.Triangles().Num());
-	TargetMeshDescription.ReserveNewVertexInstances(SourceMeshDescription.Triangles().Num() * 3);
+	const int32 TriangleCount = SourceMeshDescription.Triangles().Num();
+	const int32 VertexInstanceCount = TriangleCount * 3;
+	const int32 VertexCount = SourceMeshDescription.Vertices().Num();
+	TargetMeshDescription.ReserveNewPolygonGroups(Source.PolygonGroupMaterialSlotNames.GetNumElements());
+	TargetMeshDescription.ReserveNewPolygons(TriangleCount);
+	TargetMeshDescription.ReserveNewTriangles(TriangleCount);
+	TargetMeshDescription.ReserveNewVertexInstances(VertexInstanceCount);
+	TargetMeshDescription.ReserveNewVertices(VertexCount);
 
 	Target.UVs.SetNumChannels(Source.NumUVChannels);
 
-	TMap<FVertexID, FVertexID> SourceToTargetVertexID;
-	SourceToTargetVertexID.Reserve(SourceMeshDescription.Vertices().Num());
-
-	//Create MorphTarget Attributes:
-	TArray<bool> UseMorphTarget;
-	UseMorphTarget.SetNumUninitialized(Source.NumberOfMorphTargets);
-	TArray<bool> UseMorphTargetNormals;
-	UseMorphTargetNormals.SetNumUninitialized(Source.NumberOfMorphTargets);
-	for (const FName& SourceMorphTargetName : Source.MorphTargetNames)
-	{
-		Target.MeshAttributes.RegisterMorphTargetAttribute(SourceMorphTargetName, Source.MeshAttributes.HasMorphTargetNormalsAttribute(SourceMorphTargetName));
-		Target.MorphPositionDeltas.Add(Target.MeshAttributes.GetVertexMorphPositionDelta(SourceMorphTargetName));
-		Target.MorphNormals.Add(Target.MeshAttributes.GetVertexInstanceMorphNormalDelta(SourceMorphTargetName));
-
-		//Validations and Normal checks:
-		int32 MorphTargetIndex = Target.MorphPositionDeltas.Num() - 1;
-
-		UseMorphTarget[MorphTargetIndex] = Target.MorphPositionDeltas[MorphTargetIndex].IsValid() && Source.MorphPositionDeltas[MorphTargetIndex].IsValid();
-		UseMorphTargetNormals[MorphTargetIndex] = Target.MorphNormals[MorphTargetIndex].IsValid() && Source.MorphNormals[MorphTargetIndex].IsValid();
-	}
+	TMap<FVertexID, FVertexID> SourceToTargetVertexIDMap;
+	SourceToTargetVertexIDMap.Reserve(VertexCount);
 
 	//Copy PolygonGroups and MaterialSlots:
 	for (size_t PolygonGroupIndex = 0; PolygonGroupIndex < Source.PolygonGroupMaterialSlotNames.GetNumElements(); PolygonGroupIndex++)
@@ -1274,26 +1257,16 @@ void FSkeletalMeshOperations::FixVertexInstanceStructure(FMeshDescription& Sourc
 		//SkinWeights:
 		FVertexBoneWeights VertexBoneWeights = Source.SkinWeights.Get(SourceVertexID);
 		TArray<UE::AnimationCore::FBoneWeight> BoneWeights;
+		BoneWeights.Reserve(VertexBoneWeights.Num());
 		for (const UE::AnimationCore::FBoneWeight BoneWeight : VertexBoneWeights)
 		{
 			BoneWeights.Add(BoneWeight);
 		}
 		Target.SkinWeights.Set(TargetVertexID, BoneWeights);
 
-		//Morph Target Positions:
-		for (size_t MorphTargetIndex = 0; MorphTargetIndex < Source.NumberOfMorphTargets; MorphTargetIndex++)
-		{
-			if (!UseMorphTarget[MorphTargetIndex])
-			{
-				continue;
-			}
-			Target.MorphPositionDeltas[MorphTargetIndex][TargetVertexID] = Source.MorphPositionDeltas[MorphTargetIndex][SourceVertexID];
-		}
-
-		SourceToTargetVertexID.Add(SourceVertexID, TargetVertexID);
+		SourceToTargetVertexIDMap.Add(SourceVertexID, TargetVertexID);
 	}
 
-
 	auto CopyVertexInstance = [&](FVertexInstanceID SourceVertexInstanceID, FVertexInstanceID TargetVertexInstanceID)
 		{
 			Target.Normals[TargetVertexInstanceID] = Source.Normals[SourceVertexInstanceID];
@@ -1304,25 +1277,20 @@ void FSkeletalMeshOperations::FixVertexInstanceStructure(FMeshDescription& Sourc
 			{
 				Target.UVs.Set(TargetVertexInstanceID, UVIndex, Source.UVs.Get(SourceVertexInstanceID, UVIndex));
 			}
-
-			//Copy the potential MorphTargetNormals:
-			for (size_t MorphTargetIndex = 0; MorphTargetIndex < Source.NumberOfMorphTargets; MorphTargetIndex++)
-			{
-				if (UseMorphTarget[MorphTargetIndex] && UseMorphTargetNormals[MorphTargetIndex])
-				{
-					Target.MorphNormals[MorphTargetIndex][TargetVertexInstanceID] = Source.MorphNormals[MorphTargetIndex][SourceVertexInstanceID];
-				}
-			}
 		};
 
 	TargetFaceSmoothingMasks.SetNumZeroed(SourceMeshDescription.Triangles().Num());
-	//Vertex Instances are created to be in complete order. so we don't copy over the existing VertexInstances prior to Traingle traverser.
+	//Vertex Instances are created to be in complete order. so we don't copy over the existing VertexInstances prior to Triangle traverser.
+	TArray<TPair<int32, int32>> SourceToTargetVertexInstanceIDs;
+	SourceToTargetVertexInstanceIDs.Reserve(SourceMeshDescription.Triangles().Num() * 3);
+
 	for (FTriangleID TriangleID : SourceMeshDescription.Triangles().GetElementIDs())
 	{
 		FPolygonID PolygonID = SourceMeshDescription.GetTrianglePolygon(TriangleID);
 		FPolygonGroupID PolygonGroupID = SourceMeshDescription.GetTrianglePolygonGroup(TriangleID);
 
 		TArray<FVertexInstanceID> SourceVertexInstanceIDs = SourceMeshDescription.GetPolygonVertexInstances(PolygonID);
+			
 		TArray<FVertexInstanceID> TargetVertexInstanceIDs;
 		TargetVertexInstanceIDs.SetNum(SourceVertexInstanceIDs.Num());
 
@@ -1331,14 +1299,16 @@ void FSkeletalMeshOperations::FixVertexInstanceStructure(FMeshDescription& Sourc
 			FVertexInstanceID SourceVertexInstanceID = SourceVertexInstanceIDs[Corner];
 
 			FVertexID SourceParentVertexID = SourceMeshDescription.GetVertexInstanceVertex(SourceVertexInstanceID);
-			FVertexID TargetParentVertexID = SourceToTargetVertexID[SourceParentVertexID];
+			FVertexID TargetParentVertexID = SourceToTargetVertexIDMap[SourceParentVertexID];
 
 			FVertexInstanceID TargetVertexInstanceID = TargetMeshDescription.CreateVertexInstance(TargetParentVertexID);
 
-			TargetVertexInstanceIDs[Corner] = TargetVertexInstanceID;
-
 			//Copy VertexInstance Values:
 			CopyVertexInstance(SourceVertexInstanceID, TargetVertexInstanceID);
+
+			TargetVertexInstanceIDs[Corner] = TargetVertexInstanceID;
+			SourceToTargetVertexInstanceIDs.Add(TPair<int32, int32>(SourceVertexInstanceID, TargetVertexInstanceID)); //to be used for MorphTargetNormals copying
+				
 		}
 
 		FPolygonID TargetPolyongId = TargetMeshDescription.CreatePolygon(PolygonGroupID, TargetVertexInstanceIDs);
@@ -1353,6 +1323,50 @@ void FSkeletalMeshOperations::FixVertexInstanceStructure(FMeshDescription& Sourc
 		}
 	}
 
+	//Create MorphTarget Attributes:
+	TArray<bool> UseMorphTarget;
+	UseMorphTarget.SetNumUninitialized(Source.NumberOfMorphTargets);
+	TArray<bool> UseMorphTargetNormals;
+	UseMorphTargetNormals.SetNumUninitialized(Source.NumberOfMorphTargets);
+	for (const FName& SourceMorphTargetName : Source.MorphTargetNames)
+	{
+		Target.MeshAttributes.RegisterMorphTargetAttribute(SourceMorphTargetName, Source.MeshAttributes.HasMorphTargetNormalsAttribute(SourceMorphTargetName));
+		Target.MorphPositionDeltas.Add(Target.MeshAttributes.GetVertexMorphPositionDelta(SourceMorphTargetName));
+		Target.MorphNormals.Add(Target.MeshAttributes.GetVertexInstanceMorphNormalDelta(SourceMorphTargetName));
+
+		//Validations and Normal checks:
+		int32 MorphTargetIndex = Target.MorphPositionDeltas.Num() - 1;
+
+		UseMorphTarget[MorphTargetIndex] = Target.MorphPositionDeltas[MorphTargetIndex].IsValid() && Source.MorphPositionDeltas[MorphTargetIndex].IsValid();
+		UseMorphTargetNormals[MorphTargetIndex] = Target.MorphNormals[MorphTargetIndex].IsValid() && Source.MorphNormals[MorphTargetIndex].IsValid();
+	}
+
+	//Morph Target Positions:
+	for (size_t MorphTargetIndex = 0; MorphTargetIndex < Source.NumberOfMorphTargets; MorphTargetIndex++)
+	{
+		if (!UseMorphTarget[MorphTargetIndex])
+		{
+			continue;
+		}
+
+		for (const TPair<FVertexID, FVertexID>& Entry : SourceToTargetVertexIDMap)
+		{
+			Target.MorphPositionDeltas[MorphTargetIndex][Entry.Value] = Source.MorphPositionDeltas[MorphTargetIndex][Entry.Key];
+		}
+	}
+
+	//Copy the potential MorphTargetNormals:
+	for (size_t MorphTargetIndex = 0; MorphTargetIndex < Source.NumberOfMorphTargets; MorphTargetIndex++)
+	{
+		if (UseMorphTarget[MorphTargetIndex] && UseMorphTargetNormals[MorphTargetIndex])
+		{
+			for (const TPair<int32, int32>& Entry : SourceToTargetVertexInstanceIDs)
+			{
+				Target.MorphNormals[MorphTargetIndex][Entry.Value] = Source.MorphNormals[MorphTargetIndex][Entry.Key];
+			}
+		}
+	}
+	
 	//Copy GeometryParts:
 	for (const FSourceGeometryPartID GeometryPartID : Source.MeshAttributes.SourceGeometryParts().GetElementIDs())
 	{
@@ -1364,12 +1378,6 @@ void FSkeletalMeshOperations::FixVertexInstanceStructure(FMeshDescription& Sourc
 		Target.GeometryPartNames.Set(PartID, Name);
 		Target.GeometryPartVertexOffsetAndCounts.Set(PartID, { OffsetAndCount[0], OffsetAndCount[1]});
 	}
-
-	TargetMeshDescription.ResumeVertexInstanceIndexing();
-	TargetMeshDescription.ResumeEdgeIndexing();
-	TargetMeshDescription.ResumePolygonIndexing();
-	TargetMeshDescription.ResumePolygonGroupIndexing();
-	TargetMeshDescription.ResumeUVIndexing();
 };
 
 void FSkeletalMeshOperations::ValidateFixComputeMeshDescriptionData(FMeshDescription& MeshDescription, const TArray<uint32>& FaceSmoothingMasks, int32 LODIndex, const bool bComputeWeightedNormals, const FString& SkeletalMeshPath)
