Mesh Description: Fixed a crash when copying morph targets in FSkeletalMeshOperations::FixVertexInstanceStructure

#rnx

diff --git a/Engine/Source/Runtime/SkeletalMeshDescription/Private/SkeletalMeshOperations.cpp b/Engine/Source/Runtime/SkeletalMeshDescription/Private/SkeletalMeshOperations.cpp
--- a/Engine/Source/Runtime/SkeletalMeshDescription/Private/SkeletalMeshOperations.cpp
+++ b/Engine/Source/Runtime/SkeletalMeshDescription/Private/SkeletalMeshOperations.cpp
@@ -1184,7 +1184,7 @@ void FSkeletalMeshOperations::FixVertexInstanceStructure(FMeshDescription& Sourc
 		int32 NumberOfMorphTargets;
 		TArray<FName> MorphTargetNames; //Primarily important for the Source (or do we just move it outside of the helper?)
 		TArray<TVertexAttributesRef<FVector3f>> MorphPositionDeltas; //Vertex, SourceMeshMorphPosDeltas[0] belongs to -> SourceMorphTargetNames[0] and so on..
-		TArray<TVertexInstanceAttributesRef<FVector3f>> MorphNormals; //VertexInstance, SourceMeshMorphPosDeltas[0] belongs to -> SourceMorphTargetNames[0] and so on..
+		TArray<TVertexInstanceAttributesRef<FVector3f>> MorphNormalDeltas; //VertexInstance, SourceMeshMorphPosDeltas[0] belongs to -> SourceMorphTargetNames[0] and so on..
 
 		FMeshAttributesHelper(FMeshDescription& SourceMeshDescription, const FMeshDescription* MeshDescriptionTemplate = nullptr)
 			: MeshAttributes(SourceMeshDescription)
@@ -1231,11 +1231,11 @@ void FSkeletalMeshOperations::FixVertexInstanceStructure(FMeshDescription& Sourc
 
 			NumberOfMorphTargets = MorphTargetNames.Num();
 			MorphPositionDeltas.Reserve(NumberOfMorphTargets);
-			MorphNormals.Reserve(NumberOfMorphTargets);
+			MorphNormalDeltas.Reserve(NumberOfMorphTargets);
 			for (const FName& MorphTargetName : MorphTargetNames)
 			{
 				MorphPositionDeltas.Add(MeshAttributes.GetVertexMorphPositionDelta(MorphTargetName));
-				MorphNormals.Add(MeshAttributes.GetVertexInstanceMorphNormalDelta(MorphTargetName));
+				MorphNormalDeltas.Add(MeshAttributes.GetVertexInstanceMorphNormalDelta(MorphTargetName));
 			}
 		}
 	};
@@ -1346,31 +1346,9 @@ void FSkeletalMeshOperations::FixVertexInstanceStructure(FMeshDescription& Sourc
 		}
 	}
 
-	//Create MorphTarget Attributes:
-	TArray<bool> UseMorphTarget;
-	UseMorphTarget.SetNumUninitialized(Source.NumberOfMorphTargets);
-	TArray<bool> UseMorphTargetNormals;
-	UseMorphTargetNormals.SetNumUninitialized(Source.NumberOfMorphTargets);
-	for (const FName& SourceMorphTargetName : Source.MorphTargetNames)
-	{
-		Target.MorphPositionDeltas.Add(Target.MeshAttributes.GetVertexMorphPositionDelta(SourceMorphTargetName));
-		Target.MorphNormals.Add(Target.MeshAttributes.GetVertexInstanceMorphNormalDelta(SourceMorphTargetName));
-
-		//Validations and Normal checks:
-		int32 MorphTargetIndex = Target.MorphPositionDeltas.Num() - 1;
-
-		UseMorphTarget[MorphTargetIndex] = Target.MorphPositionDeltas[MorphTargetIndex].IsValid() && Source.MorphPositionDeltas[MorphTargetIndex].IsValid();
-		UseMorphTargetNormals[MorphTargetIndex] = Target.MorphNormals[MorphTargetIndex].IsValid() && Source.MorphNormals[MorphTargetIndex].IsValid();
-	}
-
 	//Morph Target Positions:
 	for (size_t MorphTargetIndex = 0; MorphTargetIndex < Source.NumberOfMorphTargets; MorphTargetIndex++)
 	{
-		if (!UseMorphTarget[MorphTargetIndex])
-		{
-			continue;
-		}
-
 		for (const TPair<FVertexID, FVertexID>& Entry : SourceToTargetVertexIDMap)
 		{
 			Target.MorphPositionDeltas[MorphTargetIndex][Entry.Value] = Source.MorphPositionDeltas[MorphTargetIndex][Entry.Key];
@@ -1380,11 +1358,11 @@ void FSkeletalMeshOperations::FixVertexInstanceStructure(FMeshDescription& Sourc
 	//Copy the potential MorphTargetNormals:
 	for (size_t MorphTargetIndex = 0; MorphTargetIndex < Source.NumberOfMorphTargets; MorphTargetIndex++)
 	{
-		if (UseMorphTarget[MorphTargetIndex] && UseMorphTargetNormals[MorphTargetIndex])
+		if (Source.MorphNormalDeltas[MorphTargetIndex].IsValid())
 		{
 			for (const TPair<int32, int32>& Entry : SourceToTargetVertexInstanceIDs)
 			{
-				Target.MorphNormals[MorphTargetIndex][Entry.Value] = Source.MorphNormals[MorphTargetIndex][Entry.Key];
+				Target.MorphNormalDeltas[MorphTargetIndex][Entry.Value] = Source.MorphNormalDeltas[MorphTargetIndex][Entry.Key];
 			}
 		}
 	}
