Mesh Description: Ensure we iterate over polygons rather than triangles, when reconstructing a mesh from another one.

#rnx

diff --git a/Engine/Source/Runtime/SkeletalMeshDescription/Private/SkeletalMeshOperations.cpp b/Engine/Source/Runtime/SkeletalMeshDescription/Private/SkeletalMeshOperations.cpp
--- a/Engine/Source/Runtime/SkeletalMeshDescription/Private/SkeletalMeshOperations.cpp
+++ b/Engine/Source/Runtime/SkeletalMeshDescription/Private/SkeletalMeshOperations.cpp
@@ -1248,7 +1248,7 @@ void FSkeletalMeshOperations::FixVertexInstanceStructure(FMeshDescription& Sourc
 	const int32 VertexInstanceCount = TriangleCount * 3;
 	const int32 VertexCount = SourceMeshDescription.Vertices().Num();
 	TargetMeshDescription.ReserveNewPolygonGroups(Source.PolygonGroupMaterialSlotNames.GetNumElements());
-	TargetMeshDescription.ReserveNewPolygons(TriangleCount);
+	TargetMeshDescription.ReserveNewPolygons(SourceMeshDescription.Polygons().Num());
 	TargetMeshDescription.ReserveNewTriangles(TriangleCount);
 	TargetMeshDescription.ReserveNewVertexInstances(VertexInstanceCount);
 	TargetMeshDescription.ReserveNewVertices(VertexCount);
@@ -1303,10 +1303,9 @@ void FSkeletalMeshOperations::FixVertexInstanceStructure(FMeshDescription& Sourc
 	TArray<TPair<int32, int32>> SourceToTargetVertexInstanceIDs;
 	SourceToTargetVertexInstanceIDs.Reserve(SourceMeshDescription.Triangles().Num() * 3);
 
-	for (FTriangleID TriangleID : SourceMeshDescription.Triangles().GetElementIDs())
+	for (FPolygonID PolygonID : SourceMeshDescription.Polygons().GetElementIDs())
 	{
-		FPolygonID PolygonID = SourceMeshDescription.GetTrianglePolygon(TriangleID);
-		FPolygonGroupID PolygonGroupID = SourceMeshDescription.GetTrianglePolygonGroup(TriangleID);
+		FPolygonGroupID PolygonGroupID = SourceMeshDescription.GetPolygonPolygonGroup(PolygonID);
 
 		TArray<FVertexInstanceID> SourceVertexInstanceIDs = SourceMeshDescription.GetPolygonVertexInstances(PolygonID);
 			
@@ -1337,11 +1336,20 @@ void FSkeletalMeshOperations::FixVertexInstanceStructure(FMeshDescription& Sourc
 		}
 
 		//Smooth mask:
-		if (SourceSmoothingMasks.IsValidIndex(TriangleID.GetValue()))
+		TArrayView<const FTriangleID> SourceTriangleIDs = SourceMeshDescription.GetPolygonTriangles(PolygonID);
+		TArrayView<const FTriangleID> TargetTriangleIDs = TargetMeshDescription.GetPolygonTriangles(TargetPolygonId);
+		
+		if (ensure(SourceTriangleIDs.Num() == TargetTriangleIDs.Num()))
 		{
-			if (ensure(TargetFaceSmoothingMasks.IsValidIndex(TargetPolygonId.GetValue())))
+			for (int32 TriangleIndex = 0; TriangleIndex < SourceTriangleIDs.Num(); ++TriangleIndex)
 			{
-				TargetFaceSmoothingMasks[TargetPolygonId.GetValue()] = SourceSmoothingMasks[TriangleID.GetValue()];
+				if (SourceSmoothingMasks.IsValidIndex(SourceTriangleIDs[TriangleIndex].GetValue()))
+				{
+					if (ensure(TargetFaceSmoothingMasks.IsValidIndex(TargetTriangleIDs[TriangleIndex].GetValue())))
+					{
+						TargetFaceSmoothingMasks[TargetTriangleIDs[TriangleIndex].GetValue()] = SourceSmoothingMasks[SourceTriangleIDs[TriangleIndex].GetValue()];
+					}
+				}
 			}
 		}
 	}
