Skeletal Mesh: Utility functions for bone pruning and generating bone cluster islands.
- Fixed an issue with FStaticMeshOperations::ComputeTriangleTangentsAndNormals where it would crash if the mesh didn't have a UV attribute defined.

#rnx

diff --git a/Engine/Source/Runtime/SkeletalMeshDescription/Private/SkeletalMeshAttributes.cpp b/Engine/Source/Runtime/SkeletalMeshDescription/Private/SkeletalMeshAttributes.cpp
--- a/Engine/Source/Runtime/SkeletalMeshDescription/Private/SkeletalMeshAttributes.cpp
+++ b/Engine/Source/Runtime/SkeletalMeshDescription/Private/SkeletalMeshAttributes.cpp
@@ -10,7 +10,12 @@ namespace MeshAttribute
 		const FName SkinWeights("SkinWeights");
 		const FName ImportPointIndex("ImportPointIndex");
 	}
-	
+
+	namespace PolygonGroup
+	{
+		const FName ParentPolygonGroup("ParentPolygonGroup");
+	}
+
 	namespace Bone
 	{
 		const FName Name("Name");
@@ -105,9 +110,13 @@ bool FSkeletalMeshAttributes::RegisterImportPointIndexAttribute()
 
 void FSkeletalMeshAttributes::UnregisterImportPointIndexAttribute()
 {
-	return MeshDescription.VertexAttributes().UnregisterAttribute(MeshAttribute::Vertex::ImportPointIndex);	
+	MeshDescription.VertexAttributes().UnregisterAttribute(MeshAttribute::Vertex::ImportPointIndex);	
 }
 
+bool FSkeletalMeshAttributes::RegisterPolygonGroupParentAttribute()
+{
+	return MeshDescription.PolygonGroupAttributes().RegisterAttribute<int32>(MeshAttribute::PolygonGroup::ParentPolygonGroup, 1, INDEX_NONE).IsValid();	
+}
 
 bool FSkeletalMeshAttributes::RegisterSkinWeightAttribute(const FName InProfileName)
 {
diff --git a/Engine/Source/Runtime/SkeletalMeshDescription/Private/SkeletalMeshOperations.cpp b/Engine/Source/Runtime/SkeletalMeshDescription/Private/SkeletalMeshOperations.cpp
--- a/Engine/Source/Runtime/SkeletalMeshDescription/Private/SkeletalMeshOperations.cpp
+++ b/Engine/Source/Runtime/SkeletalMeshDescription/Private/SkeletalMeshOperations.cpp
@@ -6,6 +6,7 @@
 #include "MeshDescriptionAdapter.h"
 #include "Math/GenericOctree.h"
 #include "SkeletalMeshAttributes.h"
+#include "Containers/Deque.h"
 #include "Spatial/MeshAABBTree3.h"
 
 
@@ -445,9 +446,7 @@ static void PoseMesh(
 				FSkeletalMeshOperations::ComputeTangentsAndNormals(InOutTargetMesh, EComputeNTBsFlags::WeightedNTBs);
 			
 				// We don't need the triangle tangents and normals anymore.
-				InOutTargetMesh.TriangleAttributes().UnregisterAttribute(MeshAttribute::Triangle::Normal);
-				InOutTargetMesh.TriangleAttributes().UnregisterAttribute(MeshAttribute::Triangle::Tangent);
-				InOutTargetMesh.TriangleAttributes().UnregisterAttribute(MeshAttribute::Triangle::Binormal);
+				Attributes.UnregisterTriangleNormalAndTangentAttributes();
 			}	
 		}
 	}
@@ -1931,4 +1930,351 @@ void FSkeletalMeshOperations::ApplyRigToGeo(FMeshDescription& RigMeshDescription
 	}
 }
 
+void FSkeletalMeshOperations::PruneBoneWeights(
+	FMeshDescription& InOutMeshDescription,
+	const UE::AnimationCore::FBoneWeightsSettings& InSettings,
+	TConstArrayView<FName> InProfileNames
+	)
+{
+	FSkeletalMeshAttributes Attributes(InOutMeshDescription);
+	TArray<FName> ProfileNames = Attributes.GetSkinWeightProfileNames();
+
+	if (!InProfileNames.IsEmpty())
+	{
+		ProfileNames = ProfileNames.FilterByPredicate([&InProfileNames](const FName ProfileName) -> bool
+			{
+				return InProfileNames.Contains(ProfileName);
+			});
+	}
+	
+	ParallelFor(ProfileNames.Num(), [&InOutMeshDescription, &InSettings, &ProfileNames](int32 ProfileIndex)
+		{
+			using namespace UE::AnimationCore;
+			
+			const FName ProfileName = ProfileNames[ProfileIndex];
+			FSkeletalMeshAttributes Attributes(InOutMeshDescription);
+			FSkinWeightsVertexAttributesRef SkinWeights = Attributes.GetVertexSkinWeights(ProfileName);
+
+			for (FVertexID VertexID: InOutMeshDescription.Vertices().GetElementIDs())
+			{
+				FVertexBoneWeights VertexBoneWeights = SkinWeights.Get(VertexID);
+
+				VertexBoneWeights.Renormalize(InSettings);
+			}
+		});
+}
+
+
+void FSkeletalMeshOperations::PartitionByMaximumUniqueBoneIndexCount(
+	FMeshDescription& InOutMeshDescription,
+	const int32 InMaxBoneCount,
+	const TArray<FName>& InProfileNames
+	)
+{
+	if (InOutMeshDescription.Triangles().Num() == 0)
+	{
+		return;
+	}
+
+	// If the number of unique bones can never exceed the given limit, bail out already.
+	FSkeletalMeshAttributes Attributes(InOutMeshDescription);
+	if (!Attributes.HasBones() || Attributes.GetNumBones() <= InMaxBoneCount)
+	{
+		return;
+	}
+
+	TArray<FName> ProfileNames = Attributes.GetSkinWeightProfileNames();
+
+	if (!InProfileNames.IsEmpty())
+	{
+		ProfileNames = ProfileNames.FilterByPredicate([&InProfileNames](const FName ProfileName) -> bool
+			{
+				return InProfileNames.Contains(ProfileName);
+			});
+	}
+	
+	TArray<FSkinWeightsVertexAttributesRef> SkinWeightAttributes;
+	for (FName ProfileName : ProfileNames)
+	{
+		SkinWeightAttributes.Add(Attributes.GetVertexSkinWeights(ProfileName));
+	}
+	
+	struct FTrianglePatch
+	{
+		TArray<FTriangleID> Triangles;
+		TSet<FBoneIndexType> UniqueBoneSet;
+	};
+
+	TArray<TArray<FTrianglePatch>> TrianglePatchGroups;
+	TArray<TArrayView<const FTriangleID>> GroupedTriangles;
+
+	TrianglePatchGroups.SetNum(InOutMeshDescription.PolygonGroups().GetArraySize());
+	GroupedTriangles.SetNum(InOutMeshDescription.PolygonGroups().GetArraySize());
+
+	for (FPolygonGroupID PolygonGroupID: InOutMeshDescription.PolygonGroups().GetElementIDs())
+	{
+		GroupedTriangles[PolygonGroupID] = InOutMeshDescription.GetPolygonGroupTriangles(PolygonGroupID);
+	}
+	
+	// Traverse neighboring triangles, of all the same polygon group, and add them to the same triangle patch
+	// until we exceed the list of unique bone indices for that patch, then start a new patch for the same polygon group.
+	ParallelFor(InOutMeshDescription.PolygonGroups().GetArraySize(), [&TrianglePatchGroups, &GroupedTriangles, &Mesh=InOutMeshDescription, &SkinWeightAttributes, InMaxBoneCount](int32 InGroupIndex) -> void
+		{
+			FPolygonGroupID PolygonGroupID(InGroupIndex);
+			if (!Mesh.PolygonGroups().IsValid(PolygonGroupID))
+			{
+				return;
+			}
+
+			TArrayView<const FTriangleID> GroupTriangles = GroupedTriangles[InGroupIndex];
+			if (GroupTriangles.IsEmpty())
+			{
+				return;
+			}
+			
+			TBitArray<> VisitedTriangles;
+			TDeque<FTriangleID> TriangleQueue;
+			VisitedTriangles.SetNum(Mesh.Triangles().GetArraySize(), false);
+
+			int32 CurrentTriangleIndex = 0;
+			TriangleQueue.PushLast(GroupTriangles[CurrentTriangleIndex]);
+			VisitedTriangles[GroupTriangles[CurrentTriangleIndex]] = true;
+			CurrentTriangleIndex++;
+
+			TSet<FBoneIndexType> TriangleBoneSet;
+			
+			auto PopulateBoneSetForTriangle = [&SkinWeightAttributes, &Mesh, &TriangleBoneSet](FTriangleID InTriangleID) -> void
+				{
+					TriangleBoneSet.Empty();
+					for (const FVertexID VertexID: Mesh.GetTriangleVertices(InTriangleID))
+					{
+						for (FSkinWeightsVertexAttributesRef SkinWeightAttribute: SkinWeightAttributes)
+						{
+							for (UE::AnimationCore::FBoneWeight BoneWeight: SkinWeightAttribute.Get(VertexID))
+							{
+								TriangleBoneSet.Add(BoneWeight.GetBoneIndex());
+							}
+						}
+					}
+				};
+
+			while (!TriangleQueue.IsEmpty())
+			{
+				FTriangleID TriangleID = TriangleQueue.First();
+				TriangleQueue.PopFirst();
+				FPolygonGroupID GroupID = Mesh.GetTrianglePolygonGroup(TriangleID);
+
+				TArray<FTrianglePatch>& PatchListForGroup = TrianglePatchGroups[GroupID];
+				if (PatchListForGroup.IsEmpty())
+				{
+					PatchListForGroup.AddDefaulted();
+				}
+				FTrianglePatch& CurrentPatch = PatchListForGroup.Last();
+				
+				PopulateBoneSetForTriangle(TriangleID);
+				if (CurrentPatch.UniqueBoneSet.Includes(TriangleBoneSet))
+				{
+					CurrentPatch.Triangles.Add(TriangleID);
+				}
+				else
+				{
+					// If the union of the patch's bone set and the triangle's bone set exceeds the maximum, start a new patch
+					// and reset the triangle queue. That way we can start a new seed patch that grows radially from the seed 
+					// triangle, rather than just cleaning up whatever disjointed triangles might remain in the queue.
+					TSet<FBoneIndexType> CombinedBoneSet = CurrentPatch.UniqueBoneSet;
+					CombinedBoneSet.Append(TriangleBoneSet);
+					if (CombinedBoneSet.Num() > InMaxBoneCount)
+					{
+						FTrianglePatch& NewPatch = PatchListForGroup.AddDefaulted_GetRef();
+						NewPatch.UniqueBoneSet = TriangleBoneSet;
+						NewPatch.Triangles.Add(TriangleID);
+
+						// Mark all triangles in the queue as unvisited before we clear it.
+						for (const FTriangleID QueuedTriangleID: TriangleQueue)
+						{
+							VisitedTriangles[QueuedTriangleID] = false;
+						}
+						TriangleQueue.Empty();
+						
+					}
+					else
+					{
+						CurrentPatch.UniqueBoneSet = MoveTemp(CombinedBoneSet);
+						CurrentPatch.Triangles.Add(TriangleID);
+					}
+				}
+
+				// First try to find neighbors of the same polygon group ID and add them to the queue.
+				// The reason we stick with same polygon group ID is so that we can try to maintain as
+				// large of a contiguous patch as possible. If we just spread over all possible polygon IDs
+				// the traversal would end up being more disjointed.
+				for (FEdgeID EdgeID : Mesh.GetTriangleEdges(TriangleID))
+				{
+					for (FTriangleID ConnectedTriangleID : Mesh.GetEdgeConnectedTriangleIDs(EdgeID))
+					{
+						if (ConnectedTriangleID != TriangleID && !VisitedTriangles[ConnectedTriangleID])
+						{
+							FPolygonGroupID ConnectedGroupID = Mesh.GetTrianglePolygonGroup(ConnectedTriangleID);
+							if (ConnectedGroupID == GroupID)
+							{
+								TriangleQueue.PushLast(ConnectedTriangleID);
+								VisitedTriangles[ConnectedTriangleID] = true;
+							}
+						}
+					}
+				}
+
+				// We didn't find a neighboring triangle to add. Go down the list of triangles until we find one we haven't
+				// visited yet.
+				if (TriangleQueue.IsEmpty())
+				{
+					while (CurrentTriangleIndex < GroupTriangles.Num() && VisitedTriangles[GroupTriangles[CurrentTriangleIndex]])
+					{
+						++CurrentTriangleIndex;
+					}
+
+					if (CurrentTriangleIndex != GroupTriangles.Num())
+					{
+						// Create a new patch group for this new set of triangles. This helps with merging later.
+						PatchListForGroup.AddDefaulted();
+						
+						TriangleQueue.PushLast(GroupTriangles[CurrentTriangleIndex]);
+						VisitedTriangles[GroupTriangles[CurrentTriangleIndex]] = true;
+						++CurrentTriangleIndex;
+					}
+				}
+			}
+		});
+
+	// We now have a bunch of patches. Let's see if we can combine some of them, to minimize their number.
+	ParallelFor(TrianglePatchGroups.Num(), [&TrianglePatchGroups, InMaxBoneCount](int32 InGroupIndex)
+		{
+			TArray<FTrianglePatch>& PatchListForGroup = TrianglePatchGroups[InGroupIndex];
+			if (PatchListForGroup.Num() <= 1)
+			{
+				return;
+			}
+
+			// Iteratively go through each patch pair, calculate a distance score between them and merge patches with 
+			// the lowest score until we run out of patches to merge.
+			TArray<TPair<int32 /* Patch Index */, int32 /* Distance Score */>> CandidatePatchDistances;
+			CandidatePatchDistances.Reserve(PatchListForGroup.Num());
+			
+			for (int32 PatchIndexA = 0; PatchIndexA < (PatchListForGroup.Num() - 1); ++PatchIndexA)
+			{
+				FTrianglePatch& TargetPatch = PatchListForGroup[PatchIndexA];
+				if (TargetPatch.Triangles.IsEmpty())
+				{
+					continue;
+				}
+
+				CandidatePatchDistances.Empty();
+				
+				for (int32 PatchIndexB = PatchIndexA + 1; PatchIndexB < PatchListForGroup.Num(); ++PatchIndexB)
+				{
+					FTrianglePatch& MergePatch = PatchListForGroup[PatchIndexB];
+					if (MergePatch.Triangles.IsEmpty())
+					{
+						continue;
+					}
+
+					CandidatePatchDistances.Emplace(PatchIndexB, 0);
+					for (FBoneIndexType BoneIndex : MergePatch.UniqueBoneSet)
+					{
+						if (!TargetPatch.UniqueBoneSet.Contains(BoneIndex))
+						{
+							CandidatePatchDistances.Last().Value++;
+						}
+					}
+				}
+
+				// Sort by score and merge until we exceed the unique bone count
+				CandidatePatchDistances.StableSort([](const TPair<int32, int32>& A, const TPair<int32, int32>& B)
+					{
+						return A.Value < B.Value;
+					});
+
+				for (const TPair<int32, int32>& PatchDistance: CandidatePatchDistances)
+				{
+					FTrianglePatch& MergePatch = PatchListForGroup[PatchDistance.Key];
+					TSet<FBoneIndexType> CandidateBoneSet = TargetPatch.UniqueBoneSet.Union(MergePatch.UniqueBoneSet);
+					if (CandidateBoneSet.Num() <= InMaxBoneCount)
+					{
+						TargetPatch.Triangles.Append(MergePatch.Triangles);
+						TargetPatch.UniqueBoneSet = MoveTemp(CandidateBoneSet);
+						MergePatch.Triangles.Empty();
+					}
+				}
+			}
+		});
+
+	// The number of non-empty patches, less the first one in each group, will give us the number of
+	// extra polygon groups we will need.
+	int32 RequiredNewGroupCount = 0;
+	for (const TArray<FTrianglePatch>& PatchListForGroup : TrianglePatchGroups)
+	{
+		for (int32 PatchIndex = 1; PatchIndex < PatchListForGroup.Num(); ++PatchIndex)
+		{
+			const FTrianglePatch& TrianglePatch = PatchListForGroup[PatchIndex];
+			if (!TrianglePatch.Triangles.IsEmpty())
+			{
+				RequiredNewGroupCount++;
+			}
+		}
+	}
+
+	// If there's no change to the group count, then we're good.
+	if (RequiredNewGroupCount == 0)
+	{
+		return;
+	}
+	
+	// If we have more than one patch per polygon group, re-assign all patches, save for the first one in a group, to
+	// a new polygon group but indicate where it got spawned off from.
+	Attributes.RegisterPolygonGroupParentAttribute();
+
+	InOutMeshDescription.ReserveNewPolygonGroups(RequiredNewGroupCount);
+
+	TPolygonGroupAttributesRef<FName> PolygonGroupImportedMaterialSlotNames = Attributes.GetPolygonGroupMaterialSlotNames();
+	TPolygonGroupAttributesRef<int32> PolygonGroupParents = Attributes.GetPolygonGroupParents();
+	
+	InOutMeshDescription.SuspendPolygonGroupIndexing();
+
+	// PatchGroupIndex is the actual polygon group on the original mesh.
+	for (int32 PatchGroupIndex = 0; PatchGroupIndex < TrianglePatchGroups.Num(); ++PatchGroupIndex)
+	{
+		const TArray<FTrianglePatch>& PatchListForGroup = TrianglePatchGroups[PatchGroupIndex];
+		if (PatchListForGroup.IsEmpty())
+		{
+			continue;
+		}
+
+		FPolygonGroupID PolygonGroupID(PatchGroupIndex);
+		FName MaterialName = PolygonGroupImportedMaterialSlotNames[PolygonGroupID];
+		
+		for (int32 PatchIndex = 1; PatchIndex < PatchListForGroup.Num(); ++PatchIndex)
+		{
+			const FTrianglePatch& TrianglePatch = PatchListForGroup[PatchIndex];
+			if (TrianglePatch.Triangles.IsEmpty())
+			{
+				continue;
+			}
+			
+			const FPolygonGroupID PatchPolygonGroupID = InOutMeshDescription.CreatePolygonGroup();
+			
+			PolygonGroupParents[PatchPolygonGroupID] = PolygonGroupID;
+			PolygonGroupImportedMaterialSlotNames[PatchPolygonGroupID] = MaterialName;
+
+			for (const FTriangleID TriangleID: TrianglePatch.Triangles)
+			{
+				const FPolygonID PolygonID = InOutMeshDescription.GetTrianglePolygon(TriangleID);
+				InOutMeshDescription.SetPolygonPolygonGroup(PolygonID, PatchPolygonGroupID);
+			}
+		}
+	}
+
+	InOutMeshDescription.ResumePolygonGroupIndexing();
+	InOutMeshDescription.BuildPolygonGroupIndexers();
+}
+
 #undef LOCTEXT_NAMESPACE
diff --git a/Engine/Source/Runtime/SkeletalMeshDescription/Public/SkeletalMeshAttributes.h b/Engine/Source/Runtime/SkeletalMeshDescription/Public/SkeletalMeshAttributes.h
--- a/Engine/Source/Runtime/SkeletalMeshDescription/Public/SkeletalMeshAttributes.h
+++ b/Engine/Source/Runtime/SkeletalMeshDescription/Public/SkeletalMeshAttributes.h
@@ -27,6 +27,11 @@ namespace MeshAttribute
  		extern SKELETALMESHDESCRIPTION_API const FName ImportPointIndex;
  	}
 
+	namespace PolygonGroup
+ 	{
+ 		extern SKELETALMESHDESCRIPTION_API const FName ParentPolygonGroup;
+ 	}
+
 	namespace Bone
 	{
 		extern SKELETALMESHDESCRIPTION_API const FName Name;
@@ -91,6 +96,11 @@ public:
 	
 	SKELETALMESHDESCRIPTION_API FSkeletalMeshAttributesShared(const FMeshDescription& InMeshDescription);
 
+	/**
+	 * Returns the parent polygon group read-only attribute.
+	 */
+	TPolygonGroupAttributesConstRef<int32> GetPolygonGroupParents() const { return MeshDescriptionShared.PolygonGroupAttributes().GetAttributesRef<int32>(MeshAttribute::PolygonGroup::ParentPolygonGroup); }
+	
 	//
 	// Skin Weights Methods
 	//
@@ -233,6 +243,8 @@ public:
 		return FStaticMeshAttributes::IsReservedAttributeName(InAttributeName) ||
 			   IsSkinWeightAttribute(InAttributeName) ||
 			   IsMorphTargetAttribute(InAttributeName) ||
+			   InAttributeName == MeshAttribute::Vertex::ImportPointIndex ||
+			   InAttributeName == MeshAttribute::PolygonGroup::ParentPolygonGroup ||
 			   InAttributeName == MeshAttribute::Bone::Name ||
 			   InAttributeName == MeshAttribute::Bone::ParentIndex ||
 			   InAttributeName == MeshAttribute::Bone::Pose ||
@@ -251,6 +263,18 @@ public:
 	 *  so this map cannot possibly be accurate).
 	 */
 	SKELETALMESHDESCRIPTION_API void UnregisterImportPointIndexAttribute();
+
+	/**
+	 * Register an attribute on the polygon group that can be used for polygon groups that have been split to indicate which
+	 * group they got split off from.
+	 * @see FSkeletalMeshOperations::PartitionByMaximumUniqueBoneIndexCount
+	 */
+	SKELETALMESHDESCRIPTION_API bool RegisterPolygonGroupParentAttribute();
+
+	/**
+	 * Returns the parent polygon group attribute.
+	 */
+	TPolygonGroupAttributesRef<int32> GetPolygonGroupParents() { return MeshDescription.PolygonGroupAttributes().GetAttributesRef<int32>(MeshAttribute::PolygonGroup::ParentPolygonGroup); }
 	
 	//
 	// Skin Weights Methods
diff --git a/Engine/Source/Runtime/SkeletalMeshDescription/Public/SkeletalMeshOperations.h b/Engine/Source/Runtime/SkeletalMeshDescription/Public/SkeletalMeshOperations.h
--- a/Engine/Source/Runtime/SkeletalMeshDescription/Public/SkeletalMeshOperations.h
+++ b/Engine/Source/Runtime/SkeletalMeshDescription/Public/SkeletalMeshOperations.h
@@ -5,6 +5,7 @@
 #include "BoneIndices.h"
 #include "Containers/Array.h"
 #include "CoreMinimal.h"
+#include "BoneWeights.h"
 #include "HAL/Platform.h"
 #include "Logging/LogMacros.h"
 #include "MeshTypes.h"
@@ -212,4 +213,32 @@ public:
 	 *										with current implementation we try the NearestVertices with the same principle as the FindMatchingPositionVertexIndexes. (aka based on GetSmallestDeltaBetweenTriangleLists)
 	 */
 	static SKELETALMESHDESCRIPTION_API void ApplyRigToGeo(FMeshDescription& RigMeshDescription /*Base/From/'Other'*/, FMeshDescription& GeoMeshDescription /*Target/To*/);
+
+	/**
+	 * Prune the skin weights on the given profiles to match the bone weight settings.
+	 * @param InOutMeshDescription The mesh description to modify.
+	 * @param InSettings The bone weight settings to use for controlling the pruning. 
+	 * @param InProfileNames The profile names to prune. If empty, will prune every profile. For only the default
+	 *   profile, use FSkeletalMeshAttributes::DefaultSkinWeightProfileName
+	 */
+	static SKELETALMESHDESCRIPTION_API void PruneBoneWeights(
+		FMeshDescription& InOutMeshDescription,
+		const UE::AnimationCore::FBoneWeightsSettings& InSettings,
+		TConstArrayView<FName> InProfileNames = {}
+		);
+
+	/**
+	 * Splits the mesh into further polygon groups if the number of unique bones being bound to by all the triangles
+	 * in that group exceeds the given maximum bone count. If a split happens, a new polygroup attribute is created 
+	 * ("ParentPolygonGroup") that points to the polygon group ID that the split-off group got created from.
+	 * @param InOutMeshDescription The mesh to partition.
+	 * @param InMaxBoneCount The maximum number of unique bones in each group.
+	 * @param InProfileNames The skin weight profiles to consider. If empty, will consider every profile. For only the 
+	 *   default profile, use FSkeletalMeshAttributes::DefaultSkinWeightProfileName
+	 */
+	static SKELETALMESHDESCRIPTION_API void PartitionByMaximumUniqueBoneIndexCount(
+		FMeshDescription& InOutMeshDescription,
+		const int32 InMaxBoneCount,
+		const TArray<FName>& InProfileNames = {}
+		);
 };
diff --git a/Engine/Source/Runtime/StaticMeshDescription/Private/StaticMeshAttributes.cpp b/Engine/Source/Runtime/StaticMeshDescription/Private/StaticMeshAttributes.cpp
--- a/Engine/Source/Runtime/StaticMeshDescription/Private/StaticMeshAttributes.cpp
+++ b/Engine/Source/Runtime/StaticMeshDescription/Private/StaticMeshAttributes.cpp
@@ -81,6 +81,13 @@ void FStaticMeshAttributes::RegisterTriangleNormalAndTangentAttributes()
 	MeshDescription.TriangleAttributes().RegisterAttribute<FVector3f>(MeshAttribute::Triangle::Binormal, 1, FVector3f::ZeroVector, EMeshAttributeFlags::Transient);
 }
 
+void FStaticMeshAttributes::UnregisterTriangleNormalAndTangentAttributes()
+{
+	MeshDescription.TriangleAttributes().UnregisterAttribute(MeshAttribute::Triangle::Normal);
+	MeshDescription.TriangleAttributes().UnregisterAttribute(MeshAttribute::Triangle::Tangent);
+	MeshDescription.TriangleAttributes().UnregisterAttribute(MeshAttribute::Triangle::Binormal);
+}
+
 
 void FStaticMeshAttributes::RegisterPolygonObjectNameAttribute()
 {
diff --git a/Engine/Source/Runtime/StaticMeshDescription/Private/StaticMeshOperations.cpp b/Engine/Source/Runtime/StaticMeshDescription/Private/StaticMeshOperations.cpp
--- a/Engine/Source/Runtime/StaticMeshDescription/Private/StaticMeshOperations.cpp
+++ b/Engine/Source/Runtime/StaticMeshDescription/Private/StaticMeshOperations.cpp
@@ -165,7 +165,7 @@ void FStaticMeshOperations::ComputeTriangleTangentsAndNormals(FMeshDescription&
 			TArrayView<FVector3f> TriangleTangents = Attributes.GetTriangleTangents().GetRawArray();
 			TArrayView<FVector3f> TriangleBinormals = Attributes.GetTriangleBinormals().GetRawArray();
 
-			if (Attributes.GetVertexInstanceUVs().GetNumChannels() > 0)
+			if (Attributes.GetVertexInstanceUVs().IsValid() && Attributes.GetVertexInstanceUVs().GetNumChannels() > 0)
 			{
 				VertexUVs = Attributes.GetVertexInstanceUVs().GetRawArray(0);
 			}
@@ -1395,7 +1395,7 @@ void FStaticMeshOperations::ComputeTangentsAndNormals(FMeshDescription& MeshDesc
 
 			// If the mesh has no UVs, average all tangents/bi-normals for a given vertex, rather than try to maintain
 			// the UV flow.
-			if (Attributes.GetVertexInstanceUVs().GetNumChannels() > 0)
+			if (Attributes.GetVertexInstanceUVs().IsValid() && Attributes.GetVertexInstanceUVs().GetNumChannels() > 0)
 			{
 				// Use UV0 as the base. Same as with ComputeTriangleTangentsAndNormals 
 				VertexUVs = Attributes.GetVertexInstanceUVs().GetRawArray(0);
diff --git a/Engine/Source/Runtime/StaticMeshDescription/Public/StaticMeshAttributes.h b/Engine/Source/Runtime/StaticMeshDescription/Public/StaticMeshAttributes.h
--- a/Engine/Source/Runtime/StaticMeshDescription/Public/StaticMeshAttributes.h
+++ b/Engine/Source/Runtime/StaticMeshDescription/Public/StaticMeshAttributes.h
@@ -78,6 +78,7 @@ public:
 	}	
 
 	STATICMESHDESCRIPTION_API void RegisterTriangleNormalAndTangentAttributes();
+	STATICMESHDESCRIPTION_API void UnregisterTriangleNormalAndTangentAttributes();
 
 	/** Registers a mesh attribute that is used to store the name of the object that contributed to each polygon. Used to partition the mesh
 	 *  into sub-objects, if needed.
