Mesh Description: Add utility functions to remove bones and compact the bone attribute list based on used bones.

#rnx
#Synced-CL 51657968

diff --git a/Engine/Source/Runtime/AnimationCore/Public/BoneWeights.h b/Engine/Source/Runtime/AnimationCore/Public/BoneWeights.h
--- a/Engine/Source/Runtime/AnimationCore/Public/BoneWeights.h
+++ b/Engine/Source/Runtime/AnimationCore/Public/BoneWeights.h
@@ -475,8 +475,7 @@ public:
 		FBoneIndexType InBoneIndex
 		) const
 	{
-		return Container.IndexOfByPredicate(
-			[InBoneIndex](const FBoneWeight& BW) { return InBoneIndex == BW.GetBoneIndex(); });
+		return ContainerAdapter::IndexOf(Container, [InBoneIndex](const FBoneWeight& BW) { return InBoneIndex == BW.GetBoneIndex(); });
 	}
 
 	int32 GetTypeHash() const
@@ -1150,7 +1149,7 @@ bool TBoneWeights<ContainerAdapter>::CullWeights(
 	bool bCulled = false;
 	int32 NumWeights = ContainerAdapter::Num(Container);
 
-	// If are are more entries in the container than the settings allow for, indiscriminately 
+	// If there are more entries in the container than the settings allow for, indiscriminately 
 	// remove the excess entries.
 	if (NumWeights > InSettings.GetMaxWeightCount())
 	{
@@ -1199,12 +1198,23 @@ void TBoneWeights<ContainerAdapter>::NormalizeWeights(
 	// We operate on int64, since we can easily end up with wraparound issues during one of the
 	// multiplications below when using int32. This would tank the division by WeightSum.
 	int64 WeightSum = 0;
+	bool bNeedResort = false;
+	uint16 LastWeight = std::numeric_limits<uint16>::max();
+	
 	for (int32 Index = 0; Index < NumWeights; Index++)
 	{
-		WeightSum += ContainerAdapter::Get(Container, Index).GetRawWeight();
+		const uint16 NewWeight = ContainerAdapter::Get(Container, Index).GetRawWeight();
+		WeightSum += NewWeight;
+
+		if (NewWeight > LastWeight)
+		{
+			bNeedResort = true;
+		}
+		LastWeight = NewWeight;
 	}
 
-	if ((InNormalizeType == EBoneWeightNormalizeType::Always && ensure(WeightSum != 0)) ||
+	// 
+	if ((InNormalizeType == EBoneWeightNormalizeType::Always && WeightSum != FBoneWeight::GetMaxRawWeight() && ensure(WeightSum != 0)) ||
 		WeightSum > FBoneWeight::GetMaxRawWeight())
 	{
 		int64 Correction = 0;
@@ -1217,8 +1227,7 @@ void TBoneWeights<ContainerAdapter>::NormalizeWeights(
 		// It can happen, if all the weights are the same, or similar, that the last weight may
 		// now be greater in value than the others. In that case we re-sort to ensure that our
 		// descending weight order invariant holds. 
-		bool bNeedResort = false;
-		uint16 LastWeight = std::numeric_limits<uint16>::max();
+		LastWeight = std::numeric_limits<uint16>::max();
 		for (int32 Index = 0; Index < NumWeights; Index++)
 		{
 			FBoneWeight BW = ContainerAdapter::Get(Container, Index);
@@ -1227,17 +1236,18 @@ void TBoneWeights<ContainerAdapter>::NormalizeWeights(
 			BW.SetRawWeight(NewWeight);
 			Correction = ScaledWeight - BW.GetRawWeight() * WeightSum;
 			ContainerAdapter::Set(Container, Index, BW);
+
+			// Due to error distribution, we can end up with weights out of order again.
 			if (NewWeight > LastWeight)
 			{
 				bNeedResort = true;
 			}
 			LastWeight = NewWeight;
 		}
-
-		if (bNeedResort)
-		{
-			SortWeights();
-		}
+	}
+	if (bNeedResort)
+	{
+		SortWeights();
 	}
 }
 
diff --git a/Engine/Source/Runtime/SkeletalMeshDescription/Private/SkeletalMeshOperations.cpp b/Engine/Source/Runtime/SkeletalMeshDescription/Private/SkeletalMeshOperations.cpp
--- a/Engine/Source/Runtime/SkeletalMeshDescription/Private/SkeletalMeshOperations.cpp
+++ b/Engine/Source/Runtime/SkeletalMeshDescription/Private/SkeletalMeshOperations.cpp
@@ -7,6 +7,7 @@
 #include "Math/GenericOctree.h"
 #include "SkeletalMeshAttributes.h"
 #include "Containers/Deque.h"
+#include "Containers/Queue.h"
 #include "Spatial/MeshAABBTree3.h"
 
 
@@ -86,7 +87,6 @@ void FSkeletalMeshOperations::AppendSkinWeight(const FMeshDescription& SourceMes
 	{
 		TArray<FName> TargetCustomAttributeNames;
 		TargetMesh.VertexAttributes().GetAttributeNames(TargetCustomAttributeNames);
-		int32 TargetVertexIndexOffset = FMath::Max(TargetMesh.Vertices().Num() - SourceMesh.Vertices().Num(), 0);
 
 		SourceMesh.VertexAttributes().ForEachByType<float>(UE::Private::FCreateAndCopyAttributeValues<float>(SourceMesh, TargetMesh, TargetCustomAttributeNames, AppendSettings.SourceVertexIDOffset));
 		SourceMesh.VertexAttributes().ForEachByType<FVector2f>(UE::Private::FCreateAndCopyAttributeValues<FVector2f>(SourceMesh, TargetMesh, TargetCustomAttributeNames, AppendSettings.SourceVertexIDOffset));
@@ -555,7 +555,6 @@ static void UnposeMesh(
 			const FVertexID VertexID(Index);
 			const FVertexBoneWeights BoneWeights = SkinWeightAttribute.Get(VertexID);
 			const FVector3f Position = PositionAttribute.Get(VertexID);
-			FVector3f SkinnedPosition = FVector3f::ZeroVector;
 
 			FMatrix44f SkinMatrix = FMatrix44f(FVector3f::ZeroVector, FVector3f::ZeroVector, FVector3f::ZeroVector, FVector3f::ZeroVector);
 			SkinMatrix.M[3][3] = 0.0f;
@@ -564,7 +563,7 @@ static void UnposeMesh(
 			{
 				SkinMatrix += RefToUserTransforms[BW.GetBoneIndex()] * BW.GetWeight();
 			}
-			SkinnedPosition = SkinMatrix.Inverse().TransformPosition(Position);
+			FVector3f SkinnedPosition = SkinMatrix.Inverse().TransformPosition(Position);
 			PositionAttribute.Set(VertexID, SkinnedPosition);
 
 			for (const FVertexInstanceID VertexInstanceID : Mesh.GetVertexVertexInstanceIDs(VertexID))
@@ -986,15 +985,12 @@ bool FSkeletalMeshOperations::GetUnposedMesh(
 	TArray<FMatrix44f> RefToUserTransforms;
 	TArray<FMatrix44f> RefPoseTransforms;
 
-	FSkeletalMeshAttributes::FBonePoseAttributesConstRef BonePoseAttribute = Attributes.GetBonePoses();
 	FSkeletalMeshAttributes::FBoneParentIndexAttributesConstRef ParentBoneIndexAttribute = Attributes.GetBoneParentIndices();
 	RefToUserTransforms.SetNumUninitialized(Attributes.GetNumBones());
 	RefPoseTransforms.SetNumUninitialized(Attributes.GetNumBones());
 
 	for (int32 BoneIndex = 0; BoneIndex < Attributes.GetNumBones(); BoneIndex++)
 	{
-		const int32 ParentBoneIndex = ParentBoneIndexAttribute.Get(BoneIndex);
-		
 		FTransform ReferenceBoneTransform = RefBoneTransforms[BoneIndex];
 		RefPoseTransforms[BoneIndex] = FMatrix44f{ ReferenceBoneTransform.ToMatrixWithScale() };
 		RefToUserTransforms[BoneIndex] = RefPoseTransforms[BoneIndex].Inverse() * FMatrix44f { InComponentSpaceTransforms[BoneIndex].ToMatrixWithScale() };
@@ -1050,15 +1046,12 @@ bool FSkeletalMeshOperations::GetUnposedMeshInPlace(FMeshDescription& InOutTarge
 	TArray<FMatrix44f> RefToUserTransforms;
 	TArray<FMatrix44f> RefPoseTransforms;
 
-	FSkeletalMeshAttributes::FBonePoseAttributesConstRef BonePoseAttribute = Attributes.GetBonePoses();
 	FSkeletalMeshAttributes::FBoneParentIndexAttributesConstRef ParentBoneIndexAttribute = Attributes.GetBoneParentIndices();
 	RefToUserTransforms.SetNumUninitialized(Attributes.GetNumBones());
 	RefPoseTransforms.SetNumUninitialized(Attributes.GetNumBones());
 
 	for (int32 BoneIndex = 0; BoneIndex < Attributes.GetNumBones(); BoneIndex++)
 	{
-		const int32 ParentBoneIndex = ParentBoneIndexAttribute.Get(BoneIndex);
-		
 		FTransform ReferenceBoneTransform = RefBoneTransforms[BoneIndex];
 		RefPoseTransforms[BoneIndex] = FMatrix44f{ ReferenceBoneTransform.ToMatrixWithScale() };
 		RefToUserTransforms[BoneIndex] = RefPoseTransforms[BoneIndex].Inverse() * FMatrix44f { InComponentSpaceTransforms[BoneIndex].ToMatrixWithScale() };
@@ -1601,20 +1594,25 @@ namespace RigApplicationHelpers
 			return;
 		}
 		FIndexAndZ PositionIndexAndZ(INDEX_NONE, Position);
-		int32 SortedIndex = SortedPositions.Num() / 2;
-		int32 StartIndex = 0;
+		int32 SortedIndex;
 		int32 LastTopIndex = SortedPositions.Num();
 		int32 LastBottomIndex = 0;
-		int32 SearchIterationCount = 0;
 
 		{
 			double Increments = ((double)SortedPositions[SortedPositionNumber - 1].Z - (double)SortedPositions[0].Z) / (double)SortedPositionNumber;
 
 			//Optimize the iteration count when a value is not in the middle
-			SortedIndex = FMath::RoundToInt(((double)PositionIndexAndZ.Z - (double)SortedPositions[0].Z) / Increments);
+			if (!FMath::IsNearlyZero(Increments))
+			{
+				SortedIndex = FMath::RoundToInt(((double)PositionIndexAndZ.Z - (double)SortedPositions[0].Z) / Increments);
+			}
+			else
+			{
+				SortedIndex = LastTopIndex;
+			}
 		}
 
-		for (SearchIterationCount = 0; SortedPositions.IsValidIndex(SortedIndex); ++SearchIterationCount)
+		for (int32 SearchIterationCount = 0; SortedPositions.IsValidIndex(SortedIndex); ++SearchIterationCount)
 		{
 			if (LastTopIndex - LastBottomIndex < 5)
 			{
@@ -1644,9 +1642,6 @@ namespace RigApplicationHelpers
 
 		//////////////////////////////////////////////////////////////////////////
 		//Closest point data (!bExactMatch)
-		float MinDistance = UE_MAX_FLT;
-		int32 ClosestIndex = LastBottomIndex;
-
 		for (int32 i = LastBottomIndex; i < SortedPositionNumber; i++)
 		{
 			//Get fast to the close position
@@ -1932,7 +1927,7 @@ void FSkeletalMeshOperations::ApplyRigToGeo(FMeshDescription& RigMeshDescription
 
 void FSkeletalMeshOperations::PruneBoneWeights(
 	FMeshDescription& InOutMeshDescription,
-	const UE::AnimationCore::FBoneWeightsSettings& InSettings,
+	const UE::AnimationCore::FBoneWeightsSettings& InBoneWeightSettings,
 	TConstArrayView<FName> InProfileNames
 	)
 {
@@ -1947,7 +1942,7 @@ void FSkeletalMeshOperations::PruneBoneWeights(
 			});
 	}
 	
-	ParallelFor(ProfileNames.Num(), [&InOutMeshDescription, &InSettings, &ProfileNames](int32 ProfileIndex)
+	ParallelFor(ProfileNames.Num(), [&InOutMeshDescription, &InBoneWeightSettings, &ProfileNames](int32 ProfileIndex)
 		{
 			using namespace UE::AnimationCore;
 			
@@ -1959,7 +1954,7 @@ void FSkeletalMeshOperations::PruneBoneWeights(
 			{
 				FVertexBoneWeights VertexBoneWeights = SkinWeights.Get(VertexID);
 
-				VertexBoneWeights.Renormalize(InSettings);
+				VertexBoneWeights.Renormalize(InBoneWeightSettings);
 			}
 		});
 }
@@ -1968,7 +1963,7 @@ void FSkeletalMeshOperations::PruneBoneWeights(
 void FSkeletalMeshOperations::PartitionByMaximumUniqueBoneIndexCount(
 	FMeshDescription& InOutMeshDescription,
 	const int32 InMaxBoneCount,
-	const TArray<FName>& InProfileNames
+	TConstArrayView<FName> InProfileNames
 	)
 {
 	if (InOutMeshDescription.Triangles().Num() == 0)
@@ -2277,4 +2272,333 @@ void FSkeletalMeshOperations::PartitionByMaximumUniqueBoneIndexCount(
 	InOutMeshDescription.BuildPolygonGroupIndexers();
 }
 
+
+void FSkeletalMeshOperations::RemoveBones(
+	FMeshDescription& InOutMeshDescription, 
+	TConstArrayView<FName> InBoneNames,
+	const UE::AnimationCore::FBoneWeightsSettings& InBoneWeightSettings,
+	TConstArrayView<FName> InProfileNames
+	)
+{
+	FSkeletalMeshAttributes Attributes(InOutMeshDescription);
+	
+	// No bone names to match against -- can't do any further processing.
+	if (InBoneNames.IsEmpty() || !Attributes.HasBones() || !Attributes.HasBoneNameAttribute() || !Attributes.HasBoneParentIndexAttribute())
+	{
+		return;
+	}
+	
+	// Create a map of the bone hierarchy, keyed on bone name. Each value contains the parent bone name,
+	// the original index and whether the bone is currently in use.
+	// The bone attributes are assumed to be consistent and sane.
+	FSkeletalMeshAttributes::FBoneNameAttributesRef BoneNamesAttrib = Attributes.GetBoneNames();
+	FSkeletalMeshAttributes::FBoneParentIndexAttributesRef BoneParentsAttrib = Attributes.GetBoneParentIndices();
+	
+	struct FBoneInfo
+	{
+		FName ParentBoneName;
+		bool bToRemove = false;
+		FBoneID OriginalID;
+		TArray<FName> Children;
+	};
+	
+	TMap<FName, FBoneInfo> BoneInfoMap;
+	FName RootBoneName;
+	
+	for (FBoneID BoneID: Attributes.Bones().GetElementIDs())
+	{
+		const FName BoneName = BoneNamesAttrib.Get(BoneID);
+		const int32 ParentBoneIndex = BoneParentsAttrib.Get(BoneID);
+		
+		FBoneInfo BoneInfo;
+		BoneInfo.ParentBoneName = Attributes.Bones().IsValid(ParentBoneIndex) ? BoneNamesAttrib.Get(ParentBoneIndex) : NAME_None;
+		BoneInfo.OriginalID = BoneID.GetValue();
+		BoneInfoMap.Emplace(BoneName, BoneInfo);
+		
+		if (BoneInfo.ParentBoneName == NAME_None)
+		{
+			if (RootBoneName != NAME_None)
+			{
+				UE_LOGF(LogSkeletalMeshOperations, Error, "Multiple root bones found: %ls and %ls", *RootBoneName.ToString(), *BoneName.ToString());
+				return;
+			}
+			RootBoneName = BoneName;
+		}
+	}
+	
+	// Fill in the list of children for each joint.
+	for (const TPair<FName, FBoneInfo>& Item: BoneInfoMap)
+	{
+		const FName BoneName = Item.Key;
+		const FBoneInfo& BoneInfo = Item.Value;
+		if (BoneInfo.ParentBoneName != NAME_None)
+		{
+			FBoneInfo& ParentBoneInfo = BoneInfoMap[BoneInfo.ParentBoneName];
+			ParentBoneInfo.Children.Add(BoneName);
+		}
+	}
+	
+	// For each joint in the remove list, mark it and all its children as to be removed.
+	for (const FName BoneName: InBoneNames)
+	{
+		// We cannot remove the root bone.
+		if (BoneName == RootBoneName)
+		{
+			continue;
+		}
+		
+		TQueue<FName> BoneMarkingQueue;
+		BoneMarkingQueue.Enqueue(BoneName);
+		
+		FName BoneToRemove;
+		while (BoneMarkingQueue.Peek(BoneToRemove))
+		{
+			BoneMarkingQueue.Pop();
+
+			if (FBoneInfo* BoneInfo = BoneInfoMap.Find(BoneToRemove); 
+				BoneInfo != nullptr)
+			{
+				BoneInfo->bToRemove = true;
+			
+				for (const FName ChildBoneName: BoneInfo->Children)
+				{
+					BoneMarkingQueue.Enqueue(ChildBoneName);
+				}
+			}
+		}
+	}
+
+	// Create a map that maps existing bone indexes to the non-removed bone indexes. All weights on bones
+	// that are marked for removal will be moved to the closest ancestor bone not marked for removal.
+	TMap<FBoneIndexType, FBoneIndexType> RemappedBoneIndices;
+	
+	for (const TPair<FName, FBoneInfo>& Item: BoneInfoMap)
+	{
+		const FBoneInfo* BoneInfo = &Item.Value;
+		const FBoneID FromBoneID = BoneInfo->OriginalID;
+		
+		while (BoneInfo->bToRemove)
+		{
+			check(!BoneInfo->ParentBoneName.IsNone());
+			BoneInfo = BoneInfoMap.Find(BoneInfo->ParentBoneName);
+		}
+		
+		RemappedBoneIndices.Add(FromBoneID, BoneInfo->OriginalID);
+	}
+	
+	TArray<FName> ProfileNames = Attributes.GetSkinWeightProfileNames();
+	if (!InProfileNames.IsEmpty())
+	{
+		ProfileNames = ProfileNames.FilterByPredicate([&InProfileNames](const FName ProfileName) -> bool
+			{
+				return InProfileNames.Contains(ProfileName);
+			});
+	}
+	
+	// Loop over all the skin weight profiles. We can parallelize across profiles, but not for each VertexID, since the profile
+	// data is stored as one giant dynamic array.
+	ParallelFor(ProfileNames.Num(), [&InOutMeshDescription, &InBoneWeightSettings, &RemappedBoneIndices, &ProfileNames](int32 ProfileIndex)
+		{
+			using namespace UE::AnimationCore;
+			
+			const FName ProfileName = ProfileNames[ProfileIndex];
+			FSkeletalMeshAttributes Attributes(InOutMeshDescription);
+			FSkinWeightsVertexAttributesRef SkinWeights = Attributes.GetVertexSkinWeights(ProfileName);
+
+			TMap<FBoneIndexType, uint16> BoneWeightMap;
+			TArray<FBoneWeight> BoneWeights;
+			
+			for (FVertexID VertexID: InOutMeshDescription.Vertices().GetElementIDs())
+			{
+				BoneWeightMap.Reset();
+				BoneWeights.Reset();
+				
+				FVertexBoneWeights VertexBoneWeights = SkinWeights.Get(VertexID);
+				
+				for (const FBoneWeight& BoneWeight: VertexBoneWeights)
+				{
+					if (const FBoneIndexType* RemappedBoneIndex = RemappedBoneIndices.Find(BoneWeight.GetBoneIndex()))
+					{
+						// We can sum up the values here into an uint16 without risking an overflow, since
+						// the weighs are already normalized and will never exceed max uint16.
+						BoneWeightMap.FindOrAdd(*RemappedBoneIndex) += BoneWeight.GetRawWeight();
+					}
+				}
+
+				for (const TPair<FBoneIndexType, uint16>& BoneWeightPair : BoneWeightMap)
+				{
+					BoneWeights.Add(FBoneWeight(BoneWeightPair.Key, BoneWeightPair.Value));
+				}
+				
+				SkinWeights.Set(VertexID, BoneWeights, InBoneWeightSettings);
+			}
+		});
+}
+
+
+void FSkeletalMeshOperations::RemoveUnusedBones(FMeshDescription& InOutMeshDescription)
+{
+	FSkeletalMeshAttributes Attributes(InOutMeshDescription);
+	
+	// No bone names to match against -- can't do any further processing.
+	if (!Attributes.HasBones() || !Attributes.HasBoneNameAttribute())
+	{
+		return;
+	}
+	
+	TArray<FName> ProfileNames = Attributes.GetSkinWeightProfileNames();
+	
+	TArray<TSet<FBoneIndexType>> BoneIndexSets;
+	ParallelForWithTaskContext(BoneIndexSets, ProfileNames.Num(), [&InOutMeshDescription, &ProfileNames](TSet<FBoneIndexType>& BoneIndexes, int32 ProfileIndex)
+		{
+			using namespace UE::AnimationCore;
+			
+			const FName ProfileName = ProfileNames[ProfileIndex];
+			FSkeletalMeshConstAttributes Attributes(InOutMeshDescription);
+			FSkinWeightsVertexAttributesConstRef SkinWeights = Attributes.GetVertexSkinWeights(ProfileName);
+
+			for (FVertexID VertexID: InOutMeshDescription.Vertices().GetElementIDs())
+			{
+				FVertexBoneWeightsConst VertexBoneWeights = SkinWeights.Get(VertexID);
+				
+				for (const FBoneWeight& BoneWeight: VertexBoneWeights)
+				{
+					BoneIndexes.Add(BoneWeight.GetBoneIndex());
+				}
+			}
+		});
+	
+	TSet<FBoneIndexType> UsedBoneIndexes;
+	for (const TSet<FBoneIndexType>& BoneIndexes: BoneIndexSets)
+	{
+		UsedBoneIndexes.Append(BoneIndexes);
+	}
+	
+	// Validate the indexes.
+	for (const FBoneIndexType BoneIndex: UsedBoneIndexes)
+	{
+		if (!Attributes.Bones().IsValid(FBoneID(BoneIndex)))
+		{
+			UE_LOGF(LogSkeletalMeshOperations, Error, "Found invalid bone index %d (only %d bones defined)", BoneIndex, Attributes.Bones().Num());
+			return;
+		}
+	}
+	
+	// Ensure that all ancestors of used bones are in the used list as well.
+	FSkeletalMeshAttributes::FBoneNameAttributesConstRef BoneNameAttr = Attributes.GetBoneNames();
+	FSkeletalMeshAttributes::FBoneParentIndexAttributesRef BoneParentIndexAttr = Attributes.GetBoneParentIndices();
+	TSet<FBoneID> ParentsFound;
+	int32 RootBoneIndex = INDEX_NONE;
+	for (const FBoneIndexType BoneIndex: UsedBoneIndexes)
+	{
+		FBoneID BoneID(BoneIndex);
+		ParentsFound.Reset();
+		
+		for (;;)
+		{
+			FBoneID ParentBoneID = FBoneID(BoneParentIndexAttr.Get(BoneID));
+			if (ParentsFound.Contains(ParentBoneID))
+			{
+				UE_LOGF(LogSkeletalMeshOperations, Error, "Bone hierarchy is cyclical (with bone '%ls' [%d])", 
+					*BoneNameAttr.Get(ParentBoneID).ToString(), ParentBoneID.GetValue());
+				return;
+			}
+			
+			if (ParentBoneID.GetValue() == INDEX_NONE)
+			{
+				if (RootBoneIndex != INDEX_NONE && BoneID.GetValue() != RootBoneIndex)
+				{
+					UE_LOGF(LogSkeletalMeshOperations, Error, "Found a duplicate root bone %d (first found root %d)", BoneID.GetValue(), RootBoneIndex);
+					return;
+				}
+				
+				// Found a root bone.
+				RootBoneIndex = BoneID.GetValue();
+				break;
+			}
+			if (!Attributes.Bones().IsValid(ParentBoneID))
+			{
+				UE_LOGF(LogSkeletalMeshOperations, Error, "Found invalid bone index %d (only %d bones defined)", ParentBoneID.GetValue(), Attributes.Bones().Num());
+				return;
+			}
+			
+			BoneID = ParentBoneID;
+			UsedBoneIndexes.Add(BoneID.GetValue());
+			ParentsFound.Add(BoneID);
+		}
+	}
+	
+	// Create a mapping to go from sparse indices to compact indices.
+	TArray<FBoneIndexType> BoneIndexMapping;
+	FBoneIndexType CurrentIndex = 0;
+	
+	// We need to loop over all bone attribute values, not just those that are defined as IDs, since
+	// IDs can be non-contiguous and the indices in the bones refer to the bone ID. So the mapping
+	// array needs to be able to handle all possible bone indices.
+	for (int32 BoneIndex = 0; BoneIndex < Attributes.Bones().GetArraySize(); ++BoneIndex)
+	{
+		const FBoneID BoneID(BoneIndex);
+
+		if (UsedBoneIndexes.Contains(BoneID.GetValue()))
+		{
+			BoneIndexMapping.Add(CurrentIndex++);
+		}
+		else
+		{
+			BoneIndexMapping.Add(INDEX_NONE);
+		}
+	}
+	
+	// Remap all existing bone bindings to point to the new, compacted bone IDs.
+	ParallelFor(ProfileNames.Num(), [&InOutMeshDescription, &Attributes, &BoneIndexMapping, &ProfileNames](int32 ProfileIndex)
+		{
+			using namespace UE::AnimationCore;
+			
+			const FName ProfileName = ProfileNames[ProfileIndex];
+			FSkinWeightsVertexAttributesRef SkinWeights = Attributes.GetVertexSkinWeights(ProfileName);
+			TArray<FBoneWeight> BoneWeights;
+			BoneWeights.Reserve(MaxInlineBoneWeightCount);
+
+			for (FVertexID VertexID: InOutMeshDescription.Vertices().GetElementIDs())
+			{
+				FVertexBoneWeights VertexBoneWeights = SkinWeights.Get(VertexID);
+				BoneWeights.Reset();
+				for (FBoneWeight BoneWeight: VertexBoneWeights)
+				{
+					BoneWeight.SetBoneIndex(BoneIndexMapping[BoneWeight.GetBoneIndex()]);
+					BoneWeights.Add(BoneWeight);
+				}
+				
+				SkinWeights.SetRaw(VertexID, BoneWeights);
+			}
+		});
+	
+	// Delete the unused IDs.
+	for (const FBoneID BoneID: Attributes.Bones().GetElementIDs())
+	{
+		if (UsedBoneIndexes.Contains(BoneID.GetValue()))
+		{
+			const int32 ParentBoneIndex = BoneParentIndexAttr.Get(BoneID);
+			
+			// Only remap if we have a valid remapping. Leave the index as-is otherwise (it could be indicative
+			// of a root bone, or just a very broken index).
+			if (BoneIndexMapping.IsValidIndex(ParentBoneIndex) && BoneIndexMapping[ParentBoneIndex] != INDEX_NONE)
+			{
+				BoneParentIndexAttr.Set(BoneID, BoneIndexMapping[ParentBoneIndex]);
+			}
+		}
+		else
+		{
+			// Not in the used set. Delete the bone now. We can use IDs here since deletion does not alter
+			// the value of the IDs, instead it leaves gaps in the data. Hence, the compaction below.
+			Attributes.Bones().Remove(BoneID);
+		}
+	}
+	
+	// Compact the bone attribute values to make sure the IDs of the bones are contiguous and that the
+	// indices in the parent bone attribute are correct.
+	TSparseArray<int32> CompactionRemapping;
+	Attributes.Bones().Compact(CompactionRemapping);
+}
+
 #undef LOCTEXT_NAMESPACE
diff --git a/Engine/Source/Runtime/SkeletalMeshDescription/Public/SkeletalMeshOperations.h b/Engine/Source/Runtime/SkeletalMeshDescription/Public/SkeletalMeshOperations.h
--- a/Engine/Source/Runtime/SkeletalMeshDescription/Public/SkeletalMeshOperations.h
+++ b/Engine/Source/Runtime/SkeletalMeshDescription/Public/SkeletalMeshOperations.h
@@ -217,13 +217,13 @@ public:
 	/**
 	 * Prune the skin weights on the given profiles to match the bone weight settings.
 	 * @param InOutMeshDescription The mesh description to modify.
-	 * @param InSettings The bone weight settings to use for controlling the pruning. 
+	 * @param InBoneWeightSettings The bone weight settings to use for controlling the pruning. 
 	 * @param InProfileNames The profile names to prune. If empty, will prune every profile. For only the default
 	 *   profile, use FSkeletalMeshAttributes::DefaultSkinWeightProfileName
 	 */
 	static SKELETALMESHDESCRIPTION_API void PruneBoneWeights(
 		FMeshDescription& InOutMeshDescription,
-		const UE::AnimationCore::FBoneWeightsSettings& InSettings,
+		const UE::AnimationCore::FBoneWeightsSettings& InBoneWeightSettings,
 		TConstArrayView<FName> InProfileNames = {}
 		);
 
@@ -239,6 +239,36 @@ public:
 	static SKELETALMESHDESCRIPTION_API void PartitionByMaximumUniqueBoneIndexCount(
 		FMeshDescription& InOutMeshDescription,
 		const int32 InMaxBoneCount,
-		const TArray<FName>& InProfileNames = {}
+		TConstArrayView<FName> InProfileNames = {}
+		);
+	
+	/**
+	 * Removes the specified bones, and its children, from the mesh on the given skin weight profiles. If the list of profiles
+	 * is empty, then the bones are removed from all profiles. Any bones to remove that don't exist on the mesh will be ignored.
+	 * All weights associated with removed bones are shifted to that bone's parent.
+	 * If the mesh has no bone attributes, this operation is a no-op and the mesh is left unchanged.
+	 * @note It is not allowed to remove the root bone.
+	 * 
+	 * @param InOutMeshDescription The skeletal mesh object to remove bones and bone bindings from.
+	 * @param InBoneNames The list of bone names to remove.
+	 * @param InBoneWeightSettings The settings to use when setting the new bone weights.
+	 * @param InProfileNames The list of skin weight profiles to operate on. If empty, all skin weight profiles will be
+	 *   processed.
+	 */
+	static SKELETALMESHDESCRIPTION_API void RemoveBones(
+		FMeshDescription& InOutMeshDescription,
+		TConstArrayView<FName> InBoneNames,
+		const UE::AnimationCore::FBoneWeightsSettings& InBoneWeightSettings,
+		TConstArrayView<FName> InProfileNames = {}
+		);
+
+	/**
+	 * Removes any unused bones from the bone attribute list and adjust the parent indices accordingly. All profiles are checked
+	 * for bones-in-use.
+	 * 
+	 * @param InOutMeshDescription The skeletal mesh object to remove unused bones from.
+	 */
+	static SKELETALMESHDESCRIPTION_API void RemoveUnusedBones(
+		FMeshDescription& InOutMeshDescription
 		);
 };
diff --git a/Engine/Source/Runtime/SkeletalMeshDescription/Public/SkinWeightsAttributesRef.h b/Engine/Source/Runtime/SkeletalMeshDescription/Public/SkinWeightsAttributesRef.h
--- a/Engine/Source/Runtime/SkeletalMeshDescription/Public/SkinWeightsAttributesRef.h
+++ b/Engine/Source/Runtime/SkeletalMeshDescription/Public/SkinWeightsAttributesRef.h
@@ -184,6 +184,12 @@ public:
 		return FVertexBoneWeights(AttributesRef.Get(InVertexIndex));
 	}
 
+	/**
+	 * Set the weights using properly normalized weights.
+	 * 
+	 * @param InVertexIndex The vertex index to set these weights at. 
+	 * @param InWeights The normalized weights to set.
+	 */
 	void Set(const FVertexID InVertexIndex, const UE::AnimationCore::FBoneWeights &InWeights)
 	{
 		// FBoneWeights guarantees ordering and normalization, so we just copy the raw data.
@@ -196,6 +202,14 @@ public:
 		}
 	}
 
+	/**
+	* Set bone weights using potentially unnormalized weights. The given weights will be pruned and normalized according
+	* to the provided settings. If no settings are provided, default settings will be used.
+	* 
+	* @param InVertexIndex The vertex index to set these weights at.
+	* @param InWeights The raw weights to set.
+	* @param InSettings Optional settings for pruning and normalization.
+	*/
 	void Set(
 		const FVertexID InVertexIndex, 
 		TArrayView<const UE::AnimationCore::FBoneWeight> InWeights,
@@ -208,6 +222,28 @@ public:
 		BoneWeights.SetBoneWeights(InWeights, InSettings);
 	}
 	
+	/** 
+	 * Blindly set raw weights. The user is responsible to make sure that the weights are normalized. No checking
+	 * of normalization, maximum weight count, or minimal weight threshold will be performed. 
+	 *
+	 * @param InVertexIndex The vertex index to set these weights at.
+	 * @param InWeights The raw weights to set.
+	 */
+	void SetRaw(
+		const FVertexID InVertexIndex, 
+		TArrayView<const UE::AnimationCore::FBoneWeight> InWeights
+		)
+	{
+		TArray<int32, TInlineAllocator<UE::AnimationCore::MaxInlineBoneWeightCount>> RawWeights;
+		RawWeights.Reserve(InWeights.Num());
+		
+		for (const UE::AnimationCore::FBoneWeight& Weight : InWeights)
+		{
+			RawWeights.Add(Weight.ToInt32());
+		}
+		AttributesRef.SetArrayView(InVertexIndex, RawWeights);
+	}
+	
 protected:
 	friend class FSkeletalMeshAttributes;
 
