Skeletal Mesh: Build fix for iterator invalidation.

#rnx
#Synced-CL 51657968

diff --git a/Engine/Source/Runtime/SkeletalMeshDescription/Private/SkeletalMeshOperations.cpp b/Engine/Source/Runtime/SkeletalMeshDescription/Private/SkeletalMeshOperations.cpp
--- a/Engine/Source/Runtime/SkeletalMeshDescription/Private/SkeletalMeshOperations.cpp
+++ b/Engine/Source/Runtime/SkeletalMeshDescription/Private/SkeletalMeshOperations.cpp
@@ -2420,7 +2420,7 @@ void FSkeletalMeshOperations::RemoveBones(
 					if (const FBoneIndexType* RemappedBoneIndex = RemappedBoneIndices.Find(BoneWeight.GetBoneIndex()))
 					{
 						// We can sum up the values here into an uint16 without risking an overflow, since
-						// the weighs are already normalized and will never exceed max uint16.
+						// the weights are already normalized and will never exceed max uint16.
 						BoneWeightMap.FindOrAdd(*RemappedBoneIndex) += BoneWeight.GetRawWeight();
 					}
 				}
@@ -2488,6 +2488,7 @@ void FSkeletalMeshOperations::RemoveUnusedBones(FMeshDescription& InOutMeshDescr
 	FSkeletalMeshAttributes::FBoneNameAttributesConstRef BoneNameAttr = Attributes.GetBoneNames();
 	FSkeletalMeshAttributes::FBoneParentIndexAttributesRef BoneParentIndexAttr = Attributes.GetBoneParentIndices();
 	TSet<FBoneID> ParentsFound;
+	TSet<FBoneIndexType> UsedParentBoneIndexes;
 	int32 RootBoneIndex = INDEX_NONE;
 	for (const FBoneIndexType BoneIndex: UsedBoneIndexes)
 	{
@@ -2523,11 +2524,13 @@ void FSkeletalMeshOperations::RemoveUnusedBones(FMeshDescription& InOutMeshDescr
 			}
 			
 			BoneID = ParentBoneID;
-			UsedBoneIndexes.Add(BoneID.GetValue());
-			ParentsFound.Add(BoneID);
+			UsedParentBoneIndexes.Add(BoneID.GetValue());
+			ParentsFound.Add(BoneID.GetValue());
 		}
 	}
 	
+	UsedBoneIndexes.Append(UsedParentBoneIndexes);
+	
 	// Create a mapping to go from sparse indices to compact indices.
 	TArray<FBoneIndexType> BoneIndexMapping;
 	FBoneIndexType CurrentIndex = 0;
